More fixes for latest clang failures, still fallout from Dev-Blueprints merge

diff --git a/Engine/Source/Developer/BlueprintProfiler/classes/BlueprintProfiler.cpp b/Engine/Source/Developer/BlueprintProfiler/classes/BlueprintProfiler.cpp
--- a/Engine/Source/Developer/BlueprintProfiler/classes/BlueprintProfiler.cpp
+++ b/Engine/Source/Developer/BlueprintProfiler/classes/BlueprintProfiler.cpp
@@ -4,6 +4,11 @@
 
 #include "BlueprintProfiler.h"
 
+#if WITH_EDITOR
+#include "Editor.h"
+#include "Script.h"
+#endif
+
 //////////////////////////////////////////////////////////////////////////
 // FBlueprintProfiler
 
