More fixes for latest clang failures, still fallout from Dev-Blueprints merge

diff --git a/Engine/Source/Developer/BlueprintProfiler/Public/BlueprintProfiler.h b/Engine/Source/Developer/BlueprintProfiler/Public/BlueprintProfiler.h
--- a/Engine/Source/Developer/BlueprintProfiler/Public/BlueprintProfiler.h
+++ b/Engine/Source/Developer/BlueprintProfiler/Public/BlueprintProfiler.h
@@ -12,7 +12,7 @@ class FBlueprintProfiler : public IBlueprintProfilerInterface
 public:
 
 	FBlueprintProfiler();
-	~FBlueprintProfiler();
+	virtual ~FBlueprintProfiler();
 
 	// Begin IModuleInterface
 	virtual void StartupModule() override;
diff --git a/Engine/Source/Developer/BlueprintProfiler/Public/BlueprintProfilerSupport.h b/Engine/Source/Developer/BlueprintProfiler/Public/BlueprintProfilerSupport.h
--- a/Engine/Source/Developer/BlueprintProfiler/Public/BlueprintProfilerSupport.h
+++ b/Engine/Source/Developer/BlueprintProfiler/Public/BlueprintProfilerSupport.h
@@ -2,6 +2,7 @@
 
 #pragma once
 
+#include "CoreUObject.h"
 #include "Core.h"
 #include "Script.h"
 
