Fix compile breakage on Linux (UE-23918).

- Classes directory was mistakenly named "classes" in CL 2724345. This is a bandaid fix, proper resolution is tracked as UE-23934.
- Note that renaming is pretty much impossible due to Perforce being unable to handle that over multiple branches.

diff --git a/Engine/Source/Developer/BlueprintProfiler/BlueprintProfiler.Build.cs b/Engine/Source/Developer/BlueprintProfiler/BlueprintProfiler.Build.cs
--- a/Engine/Source/Developer/BlueprintProfiler/BlueprintProfiler.Build.cs
+++ b/Engine/Source/Developer/BlueprintProfiler/BlueprintProfiler.Build.cs
@@ -21,6 +21,12 @@ public class BlueprintProfiler : ModuleRules
         {
             PrivateDependencyModuleNames.Add("UnrealEd");
         }
+
+        if (UnrealBuildTool.BuildHostPlatform.Current.Platform == UnrealTargetPlatform.Linux)
+        {
+            // work around for bad naming of "classes" directory, proper resolution is tracked as UE-23934
+            bFasterWithoutUnity = true;
+        }
     }
 }
 
