Removing debug code (a #pragma optimize) that snuck in with the latest Dev-BP copy.

diff --git a/Engine/Source/Developer/BlueprintProfiler/Private/BlueprintProfiler.cpp b/Engine/Source/Developer/BlueprintProfiler/Private/BlueprintProfiler.cpp
--- a/Engine/Source/Developer/BlueprintProfiler/Private/BlueprintProfiler.cpp
+++ b/Engine/Source/Developer/BlueprintProfiler/Private/BlueprintProfiler.cpp
@@ -5,8 +5,6 @@
 
 #define LOCTEXT_NAMESPACE "BlueprintProfiler"
 
-#pragma optimize( "g", off )
-
 //////////////////////////////////////////////////////////////////////////
 // FBlueprintProfiler
 
