Copying //UE4/Dev-Blueprints to Dev-Main (//UE4/Dev-Main)

==========================
MAJOR FEATURES + CHANGES
==========================

Change 2884592 on 2016/02/27 by [REDACTED_PERSON]

	Packages containing Dynamic Types are listed as dependencies in FAsyncPackage::LoadImports

Change 2884607 on 2016/02/27 by [REDACTED_PERSON]

	CDO of DynamicClass is postponed as a regular CDO creation.

	This change is risky (it still requires some tests), but it seems to be necessary for solving (cyclic) dependencies while Async Loading.

Change 2885915 on 2016/02/29 by [REDACTED_PERSON]

	Struct pins on exposed on spawn properties will no longer error on compile if they do not have a pin connected to them.

	UKismetSystemLibrary::SetStructurePropertyByName's Value pin is now marked as a AutoCreateRefTerm so literals can be used on the auto-generated node.

	Modified FKismetCompilerUtilities::GenerateAssignmentNodes to assign the literal value into the pin, which in turn allows it to expand into an auto-ref term.

Change 2887269 on 2016/03/01 by [REDACTED_PERSON]

	Fixes related to DisregardForGC:
	- merged 2885687 from Dev-Core branch
	- ensure CDO of newly added class is created in CloseDisregardForGC

Change 2887273 on 2016/03/01 by [REDACTED_PERSON]

	GUObjectArray.CloseDisregardForGC(); is called before GUObjectArray.DisableDisregardForGC();

Change 2892502 on 2016/03/03 by [REDACTED_PERSON]

	More descriptive ensures.

Change 2892509 on 2016/03/03 by [REDACTED_PERSON]

	Minor changes in Orion code, necessary to compile the project with nativized Blueprints

Change 2892513 on 2016/03/03 by [REDACTED_PERSON]

	Blueprint C++ Conversion: there is no crash when a asset used by a nativized class wasn't loaded.

Change 2894347 on 2016/03/04 by [REDACTED_PERSON]

	Fixed crash with the Array Item Get node in the disassembler.

Change 2895311 on 2016/03/04 by [REDACTED_PERSON]

	Right click and using "Find References" on an event node in either the MyBlueprint window or a graph panel, will search the function's name, not the node's title (drops the Event text).

	PR #2092: Skip "Event " Prefix when Finding References for Events (Contributed by [REDACTED_PERSON])

Change 2896714 on 2016/03/07 by [REDACTED_PERSON]

	Fixes for a few related crash/assert issues when profiling blueprints that use event dispatchers/call event functions and latent nodes.
	#UE-27090 - Crash when calling an Event Dispatcher within a Sequence node with the Blueprint Profiler on
	#Proj KismetCompiler, BlueprintProfiler, BlueprintGraph, Kismet, CoreUObject

Change 2897335 on 2016/03/07 by [REDACTED_PERSON]

	Unshelved from pending changelist '2889006':

	We now copy UObjects that are assigned to Instance properties via  UObjectPropertyBase::ImportText_Internal

Change 2899151 on 2016/03/08 by [REDACTED_PERSON]

	[UEBP-112] BP profiler - pure node timings

	change summary:
	- first-pass revisions for tracking/visualizing pure node execution timings while profiling
	- trace path debugging support (default: off)

Change 2901763 on 2016/03/09 by [REDACTED_PERSON]

	Can undo/redo nodes added using shortcuts.

	Also fixed issue with macro nodes added using shortcuts not being selected.

Change 2902762 on 2016/03/10 by [REDACTED_PERSON]

	[UE-28167] Fix compile-time crash caused by intermediate pure nodes when the BP profiler view is active.

	change summary:
	- modified FBlueprintExecutionContext::MapNodeExecution() to avoid impure nodes when mapping the pure node execution chain.

Change 2907961 on 2016/03/14 by [REDACTED_PERSON]

	Manually integrated CL#2906835 from Dev-Core - Reset last non-gc index when disabling disregard for GC pool

Change 2908013 on 2016/03/14 by [REDACTED_PERSON]

	EmptyLinkFunctionForGeneratedCode function for all .generated.*.cpp files
	Fixed problem when stuff from some .generated.*.cpp files for nativized Blueprints plugin were not linked at all.

Change 2908033 on 2016/03/14 by [REDACTED_PERSON]

	Various fixes and improvements related to initialization sequence of converted Dynamic Classes.

diff --git a/Engine/Source/Developer/BlueprintCompilerCppBackend/Private/BlueprintCompilerCppBackend.cpp b/Engine/Source/Developer/BlueprintCompilerCppBackend/Private/BlueprintCompilerCppBackend.cpp
--- a/Engine/Source/Developer/BlueprintCompilerCppBackend/Private/BlueprintCompilerCppBackend.cpp
+++ b/Engine/Source/Developer/BlueprintCompilerCppBackend/Private/BlueprintCompilerCppBackend.cpp
@@ -834,7 +834,7 @@ void FBlueprintCompilerCppBackend::InnerFunctionImplementation(FKismetFunctionCo
 	{
 		bUseGotoState = FunctionContext.MustUseSwitchState(nullptr) || FunctionContext.bIsUbergraph;
 	}
-	ensure(!bUseFlowStack || bUseGotoState);
+	ensureMsgf(!bUseFlowStack || bUseGotoState, TEXT("FBlueprintCompilerCppBackend::InnerFunctionImplementation - %s"), *GetPathNameSafe(FunctionContext.Function));
 	TArray<UEdGraphNode*>* ActualLinearExecutionList = &FunctionContext.LinearExecutionList;
 	if (bUseGotoState)
 	{
@@ -1074,7 +1074,10 @@ void FBlueprintCompilerCppBackend::EmitAllStatements(FKismetFunctionContext &Fun
 	{
 		UEdGraphNode* StatementNode = LinearExecutionList[NodeIndex];
 		TArray<FBlueprintCompiledStatement*>* StatementList = FunctionContext.StatementsPerNode.Find(StatementNode);
-		ensure(StatementNode && StatementNode->IsA<UK2Node>() && !CastChecked<UK2Node>(StatementNode)->IsNodePure());
+		ensureMsgf(StatementNode && StatementNode->IsA<UK2Node>() && !CastChecked<UK2Node>(StatementNode)->IsNodePure()
+			, TEXT("Wrong Statement node %s in function %s")
+			, *GetPathNameSafe(StatementNode)
+			, *GetPathNameSafe(FunctionContext.Function));
 
 		const bool bIsCurrentExecutionGroup = !bUseExecutionGroup || FunctionContext.UnsortedSeparateExecutionGroups[ExecutionGroup].Contains(StatementNode);
 		if (StatementList && bIsCurrentExecutionGroup)
diff --git a/Engine/Source/Developer/BlueprintCompilerCppBackend/Private/BlueprintCompilerCppBackendAnim.cpp b/Engine/Source/Developer/BlueprintCompilerCppBackend/Private/BlueprintCompilerCppBackendAnim.cpp
--- a/Engine/Source/Developer/BlueprintCompilerCppBackend/Private/BlueprintCompilerCppBackendAnim.cpp
+++ b/Engine/Source/Developer/BlueprintCompilerCppBackend/Private/BlueprintCompilerCppBackendAnim.cpp
@@ -19,7 +19,7 @@ void FBackendHelperAnim::CreateAnimClassData(FEmitterLocalContext& Context)
 	if (auto AnimClass = Cast<UAnimBlueprintGeneratedClass>(Context.GetCurrentlyGeneratedClass()))
 	{
 		const FString LocalNativeName = Context.GenerateUniqueLocalName();
-		Context.AddLine(FString::Printf(TEXT("auto %s = NewObject<UAnimClassData>(GetClass(), TEXT(\"AnimClassData\"));"), *LocalNativeName));
+		Context.AddLine(FString::Printf(TEXT("auto %s = NewObject<UAnimClassData>(InDynamicClass, TEXT(\"AnimClassData\"));"), *LocalNativeName));
 
 		auto AnimClassData = NewObject<UAnimClassData>(GetTransientPackage(), TEXT("AnimClassData"));
 		AnimClassData->CopyFrom(AnimClass);
@@ -33,7 +33,7 @@ void FBackendHelperAnim::CreateAnimClassData(FEmitterLocalContext& Context)
 				, FEmitDefaultValueHelper::EPropertyAccessOperator::Pointer);
 		}
 
-		Context.AddLine(FString::Printf(TEXT("CastChecked<UDynamicClass>(GetClass())->AnimClassImplementation = %s;"), *LocalNativeName));
+		Context.AddLine(FString::Printf(TEXT("InDynamicClass->AnimClassImplementation = %s;"), *LocalNativeName));
 	}
 }
 
diff --git a/Engine/Source/Developer/BlueprintCompilerCppBackend/Private/BlueprintCompilerCppBackendBase.cpp b/Engine/Source/Developer/BlueprintCompilerCppBackend/Private/BlueprintCompilerCppBackendBase.cpp
--- a/Engine/Source/Developer/BlueprintCompilerCppBackend/Private/BlueprintCompilerCppBackendBase.cpp
+++ b/Engine/Source/Developer/BlueprintCompilerCppBackend/Private/BlueprintCompilerCppBackendBase.cpp
@@ -233,7 +233,7 @@ FString FBlueprintCompilerCppBackendBase::GenerateCodeFromClass(UClass* SourceCl
 			AdditionalMD.Add(ReplaceConvertedMD);
 		}
 
-		AdditionalMD.Add(FString::Printf(TEXT("CustomDynamicClassInitialization=\"%s::__GatherReferencedConvertedFields\""), *CppClassName));
+		// AdditionalMD.Add(FString::Printf(TEXT("CustomDynamicClassInitialization=\"%s::__CustomDynamicClassInitialization\""), *CppClassName));
 
 		EmitterContext.Header.AddLine(FString::Printf(TEXT("UCLASS(%s%s)")
 			, (!SourceClass->IsChildOf<UBlueprintFunctionLibrary>()) ? TEXT("Blueprintable, BlueprintType, ") : TEXT("")
@@ -268,7 +268,7 @@ FString FBlueprintCompilerCppBackendBase::GenerateCodeFromClass(UClass* SourceCl
 		EmitterContext.Header.AddLine(FString::Printf(TEXT("%s(const FObjectInitializer& ObjectInitializer = FObjectInitializer::Get());"), *CppClassName));
 		EmitterContext.Header.AddLine(TEXT("virtual void PostLoadSubobjects(FObjectInstancingGraph* OuterInstanceGraph) override;"));
 		EmitterContext.Header.AddLine(TEXT("static void __StaticDependenciesAssets(TArray<FBlueprintDependencyData>& AssetsToLoad);"));
-		EmitterContext.Header.AddLine(TEXT("static void __GatherReferencedConvertedFields(UDynamicClass* InDynamicClass);"));
+		EmitterContext.Header.AddLine(TEXT("static void __CustomDynamicClassInitialization(UDynamicClass* InDynamicClass);"));
 
 		FEmitDefaultValueHelper::GenerateConstructor(EmitterContext);
 	}
diff --git a/Engine/Source/Developer/BlueprintCompilerCppBackend/Private/BlueprintCompilerCppBackendUtils.cpp b/Engine/Source/Developer/BlueprintCompilerCppBackend/Private/BlueprintCompilerCppBackendUtils.cpp
--- a/Engine/Source/Developer/BlueprintCompilerCppBackend/Private/BlueprintCompilerCppBackendUtils.cpp
+++ b/Engine/Source/Developer/BlueprintCompilerCppBackend/Private/BlueprintCompilerCppBackendUtils.cpp
@@ -101,7 +101,7 @@ FString FEmitterLocalContext::FindGloballyMappedObject(const UObject* Object, co
 
 	if (ActualClass && ((Object == ActualClass) || (Object == OriginalActualClass)))
 	{
-		return CastCustomClass(TEXT("GetClass()"));
+		return CastCustomClass(((CurrentCodeType == EGeneratedCodeType::SubobjectsOfClass) ? TEXT("InDynamicClass") : TEXT("GetClass()")));
 	}
 
 	{
@@ -156,7 +156,7 @@ FString FEmitterLocalContext::FindGloballyMappedObject(const UObject* Object, co
 			const int32 AssetIndex = Dependencies.Assets.IndexOfByKey(Object);
 			if (INDEX_NONE != AssetIndex)
 			{
-				return FString::Printf(TEXT("CastChecked<%s>(CastChecked<UDynamicClass>(%s::StaticClass())->UsedAssets[%d])")
+				return FString::Printf(TEXT("CastChecked<%s>(CastChecked<UDynamicClass>(%s::StaticClass())->UsedAssets[%d], ECastCheckedType::NullAllowed)")
 					, *ClassString()
 					, *FEmitHelper::GetCppName(ActualClass)
 					, AssetIndex);
diff --git a/Engine/Source/Developer/BlueprintCompilerCppBackend/Private/BlueprintCompilerCppBackendUtils.h b/Engine/Source/Developer/BlueprintCompilerCppBackend/Private/BlueprintCompilerCppBackendUtils.h
--- a/Engine/Source/Developer/BlueprintCompilerCppBackend/Private/BlueprintCompilerCppBackendUtils.h
+++ b/Engine/Source/Developer/BlueprintCompilerCppBackend/Private/BlueprintCompilerCppBackendUtils.h
@@ -272,6 +272,8 @@ struct FEmitDefaultValueHelper
 
 	static void GenerateConstructor(FEmitterLocalContext& Context);
 
+	static void GenerateCustomDynamicClassInitialization(FEmitterLocalContext& Context);
+
 	enum class EPropertyAccessOperator
 	{
 		None, // for self scope, this
diff --git a/Engine/Source/Developer/BlueprintCompilerCppBackend/Private/BlueprintCompilerCppBackendValueHelper.cpp b/Engine/Source/Developer/BlueprintCompilerCppBackend/Private/BlueprintCompilerCppBackendValueHelper.cpp
--- a/Engine/Source/Developer/BlueprintCompilerCppBackend/Private/BlueprintCompilerCppBackendValueHelper.cpp
+++ b/Engine/Source/Developer/BlueprintCompilerCppBackend/Private/BlueprintCompilerCppBackendValueHelper.cpp
@@ -636,56 +636,49 @@ public:
 		return FString(TEXT("Z_Construct_")) + ClassString + TEXT("_") + Result + TEXT("()");
 	}
 
-	static void AddDependenciesInConstructor(FEmitterLocalContext& Context)
-	{
-		if (Context.Dependencies.Assets.Num())
-		{
-			Context.AddLine(TEXT("// List of all referenced assets"));
-		}
-		for (auto LocAsset : Context.Dependencies.Assets)
-		{
-			const FString AssetStr = Context.FindGloballyMappedObject(LocAsset, UObject::StaticClass(), true, false);
-			Context.AddLine(FString::Printf(TEXT("CastChecked<UDynamicClass>(GetClass())->UsedAssets.Add(%s);"), *AssetStr));
-		}
-	}
-
 	static void AddStaticFunctionsForDependencies(FEmitterLocalContext& Context)
 	{
 		auto SourceClass = Context.GetCurrentlyGeneratedClass();
 		auto OriginalClass = Context.Dependencies.FindOriginalClass(SourceClass);
 		const FString CppClassName = FEmitHelper::GetCppName(OriginalClass);
 		
-
 		// __StaticDependenciesAssets
 		Context.AddLine(FString::Printf(TEXT("void %s::__StaticDependenciesAssets(TArray<FBlueprintDependencyData>& AssetsToLoad)"), *CppClassName));
 		Context.AddLine(TEXT("{"));
 		Context.IncreaseIndent();
 
+		auto CreateAssetToLoadString = [&](const UObject* AssetObj) -> FString
+		{
+			UClass* AssetType = AssetObj->GetClass();
+			if (AssetType->IsChildOf<UUserDefinedEnum>())
+			{
+				AssetType = UEnum::StaticClass();
+			}
+			else if (AssetType->IsChildOf<UUserDefinedStruct>())
+			{
+				AssetType = UScriptStruct::StaticClass();
+			}
+			else if (AssetType->IsChildOf<UBlueprintGeneratedClass>() && Context.Dependencies.WillClassBeConverted(CastChecked<UBlueprintGeneratedClass>(AssetObj)))
+			{
+				AssetType = UDynamicClass::StaticClass();
+			}
+
+			return FString::Printf(TEXT("AssetsToLoad.Add({FName(TEXT(\"%s\")), FName(TEXT(\"%s\")), FName(TEXT(\"%s\")), FName(TEXT(\"%s\"))});")
+				, *AssetObj->GetOutermost()->GetPathName()
+				, *AssetObj->GetName()
+				, *AssetType->GetOutermost()->GetPathName()
+				, *AssetType->GetName());
+		};
 		for (UObject* LocAsset : Context.Dependencies.Assets)
 		{
-			auto GetConvertedClass = [&](const UClass* AssetType) -> const UClass*
+			Context.AddLine(CreateAssetToLoadString(LocAsset));
+		}
+		for (UObject* LocAsset : Context.Dependencies.ConvertedClasses)
+		{
+			if (!Context.Dependencies.Assets.Contains(LocAsset))
 			{
-				const UBlueprintGeneratedClass* BPGC = Cast<const UBlueprintGeneratedClass>(AssetType);
-				if (AssetType->IsChildOf<UUserDefinedEnum>())
-				{
-					return UEnum::StaticClass();
-				}
-				else if (AssetType->IsChildOf<UUserDefinedStruct>())
-				{
-					return UScriptStruct::StaticClass();
-				}
-				else if (BPGC && Context.Dependencies.WillClassBeConverted(BPGC))
-				{
-					return UDynamicClass::StaticClass();
-				}
-				return AssetType;
-			};
-			const UClass* ClassToUse = GetConvertedClass(LocAsset->GetClass());
-			Context.AddLine(FString::Printf(TEXT("AssetsToLoad.Add({FName(TEXT(\"%s\")), FName(TEXT(\"%s\")), FName(TEXT(\"%s\")), FName(TEXT(\"%s\"))});")
-				, *LocAsset->GetOutermost()->GetPathName()
-				, *LocAsset->GetName()
-				, *ClassToUse->GetOutermost()->GetPathName()
-				, *ClassToUse->GetName()));
+				Context.AddLine(CreateAssetToLoadString(LocAsset));
+			}
 		}
 
 		Context.DecreaseIndent();
@@ -718,101 +711,150 @@ public:
 	}
 };
 
-void FEmitDefaultValueHelper::GenerateConstructor(FEmitterLocalContext& Context)
+void FEmitDefaultValueHelper::GenerateCustomDynamicClassInitialization(FEmitterLocalContext& Context)
 {
 	auto BPGC = CastChecked<UBlueprintGeneratedClass>(Context.GetCurrentlyGeneratedClass());
 	const FString CppClassName = FEmitHelper::GetCppName(BPGC);
 
-	UClass* SuperClass = BPGC->GetSuperClass();
-	const bool bSuperHasOnlyDefaultConstructor = SuperClass && SuperClass->HasMetaData(TEXT("OnlyDefaultConstructorDeclared"));
-	Context.AddLine(FString::Printf(TEXT("%s::%s(const FObjectInitializer& ObjectInitializer) : Super(%s)")
-		, *CppClassName
-		, *CppClassName
-		, bSuperHasOnlyDefaultConstructor ? TEXT("") : TEXT("ObjectInitializer")));
+	Context.AddLine(FString::Printf(TEXT("void %s::__CustomDynamicClassInitialization(UDynamicClass* InDynamicClass)"), *CppClassName));
 	Context.AddLine(TEXT("{"));
 	Context.IncreaseIndent();
+	Context.AddLine(TEXT("ensure(0 == InDynamicClass->ReferencedConvertedFields.Num());"));
+	Context.AddLine(TEXT("ensure(0 == InDynamicClass->MiscConvertedSubobjects.Num());"));
+	Context.AddLine(TEXT("ensure(0 == InDynamicClass->UsedAssets.Num());"));
+	Context.AddLine(TEXT("ensure(0 == InDynamicClass->DynamicBindingObjects.Num());"));
+	Context.AddLine(TEXT("ensure(0 == InDynamicClass->ComponentTemplates.Num());"));
+	Context.AddLine(TEXT("ensure(0 == InDynamicClass->Timelines.Num());"));
+	Context.AddLine(TEXT("ensure(nullptr == InDynamicClass->AnimClassImplementation);"));
+
+	Context.CurrentCodeType = FEmitterLocalContext::EGeneratedCodeType::SubobjectsOfClass;
+	Context.ResetPropertiesForInaccessibleStructs();
 
-	// When CDO is created create all subobjects owned by the class
+	if (Context.Dependencies.ConvertedClasses.Num())
 	{
-		TArray<UActorComponent*> ActorComponentTempatesOwnedByClass = BPGC->ComponentTemplates;
-		// Gather all CT from SCS and IH, the remaining ones are generated for class..
-		if (auto SCS = BPGC->SimpleConstructionScript)
+		Context.AddLine(TEXT("// List of all referenced converted classes"));
+	}
+	for (auto LocStruct : Context.Dependencies.ConvertedClasses)
+	{
+		UClass* ClassToLoad = Context.Dependencies.FindOriginalClass(LocStruct);
+		if (ensure(ClassToLoad))
 		{
-			for (auto Node : SCS->GetAllNodes())
-			{
-				ActorComponentTempatesOwnedByClass.RemoveSwap(Node->ComponentTemplate);
-			}
+			const FString ClassConstructor = FDependenciesHelper::GenerateZConstructor(ClassToLoad);
+			Context.AddLine(FString::Printf(TEXT("extern UClass* %s;"), *ClassConstructor));
+			Context.AddLine(FString::Printf(TEXT("InDynamicClass->ReferencedConvertedFields.Add(%s);"), *ClassConstructor));
+
+			//Context.AddLine(FString::Printf(TEXT("InDynamicClass->ReferencedConvertedFields.Add(LoadObject<UClass>(nullptr, TEXT(\"%s\")));")
+			//	, *(ClassToLoad->GetPathName().ReplaceCharWithEscapedChar())));
 		}
-		if (auto IH = BPGC->GetInheritableComponentHandler())
+	}
+
+	if (Context.Dependencies.ConvertedStructs.Num())
+	{
+		Context.AddLine(TEXT("// List of all referenced converted structures"));
+	}
+	for (auto LocStruct : Context.Dependencies.ConvertedStructs)
+	{
+		const FString StructConstructor = FDependenciesHelper::GenerateZConstructor(LocStruct);
+		Context.AddLine(FString::Printf(TEXT("extern UScriptStruct* %s;"), *StructConstructor));
+		Context.AddLine(FString::Printf(TEXT("InDynamicClass->ReferencedConvertedFields.Add(%s);"), *StructConstructor));
+	}
+
+	TArray<UActorComponent*> ActorComponentTempatesOwnedByClass = BPGC->ComponentTemplates;
+	// Gather all CT from SCS and IH, the remaining ones are generated for class..
+	if (auto SCS = BPGC->SimpleConstructionScript)
+	{
+		for (auto Node : SCS->GetAllNodes())
 		{
-			TArray<UActorComponent*> AllTemplates;
-			IH->GetAllTemplates(AllTemplates);
-			ActorComponentTempatesOwnedByClass.RemoveAllSwap([&](UActorComponent* Component) -> bool
-			{
-				return AllTemplates.Contains(Component);
-			});
+			ActorComponentTempatesOwnedByClass.RemoveSwap(Node->ComponentTemplate);
 		}
+	}
+	if (auto IH = BPGC->GetInheritableComponentHandler())
+	{
+		TArray<UActorComponent*> AllTemplates;
+		IH->GetAllTemplates(AllTemplates);
+		ActorComponentTempatesOwnedByClass.RemoveAllSwap([&](UActorComponent* Component) -> bool
+		{
+			return AllTemplates.Contains(Component);
+		});
+	}
 
-		Context.AddLine(FString::Printf(TEXT("if(HasAnyFlags(RF_ClassDefaultObject) && (%s::StaticClass() == GetClass()))"), *CppClassName));
-		Context.AddLine(TEXT("{"));
-		Context.IncreaseIndent();
-		Context.AddLine(TEXT("ensure(0 == CastChecked<UDynamicClass>(GetClass())->MiscConvertedSubobjects.Num());"));
-		Context.AddLine(TEXT("ensure(0 == CastChecked<UDynamicClass>(GetClass())->UsedAssets.Num());"));
-		Context.AddLine(TEXT("ensure(0 == CastChecked<UDynamicClass>(GetClass())->DynamicBindingObjects.Num());"));
-		Context.AddLine(TEXT("ensure(0 == CastChecked<UDynamicClass>(GetClass())->ComponentTemplates.Num());"));
-		Context.AddLine(TEXT("ensure(0 == CastChecked<UDynamicClass>(GetClass())->Timelines.Num());"));
-		Context.AddLine(TEXT("ensure(nullptr == CastChecked<UDynamicClass>(GetClass())->AnimClassImplementation);"));
-		Context.CurrentCodeType = FEmitterLocalContext::EGeneratedCodeType::SubobjectsOfClass;
-		Context.ResetPropertiesForInaccessibleStructs();
-
-		FDependenciesHelper::AddDependenciesInConstructor(Context);
-
-		auto CreateAndInitializeClassSubobjects = [&](bool bCreate, bool bInitilize)
+	if (Context.Dependencies.Assets.Num())
+	{
+		Context.AddLine(TEXT("// List of all referenced assets"));
+	}
+	for (auto LocAsset : Context.Dependencies.Assets)
+	{
+		const FString AssetStr = Context.FindGloballyMappedObject(LocAsset, UObject::StaticClass(), true, false);
+		Context.AddLine(FString::Printf(TEXT("InDynamicClass->UsedAssets.Add(%s);"), *AssetStr));
+	}
+
+	auto CreateAndInitializeClassSubobjects = [&](bool bCreate, bool bInitilize)
+	{
+		for (auto ComponentTemplate : ActorComponentTempatesOwnedByClass)
 		{
-			for (auto ComponentTemplate : ActorComponentTempatesOwnedByClass)
+			if (ComponentTemplate)
 			{
-				if (ComponentTemplate)
-				{
-					HandleClassSubobject(Context, ComponentTemplate, FEmitterLocalContext::EClassSubobjectList::ComponentTemplates, bCreate, bInitilize);
-				}
+				HandleClassSubobject(Context, ComponentTemplate, FEmitterLocalContext::EClassSubobjectList::ComponentTemplates, bCreate, bInitilize);
 			}
+		}
 
-			for (auto TimelineTemplate : BPGC->Timelines)
+		for (auto TimelineTemplate : BPGC->Timelines)
+		{
+			if (TimelineTemplate)
 			{
-				if (TimelineTemplate)
-				{
-					HandleClassSubobject(Context, TimelineTemplate, FEmitterLocalContext::EClassSubobjectList::Timelines, bCreate, bInitilize);
-				}
+				HandleClassSubobject(Context, TimelineTemplate, FEmitterLocalContext::EClassSubobjectList::Timelines, bCreate, bInitilize);
 			}
+		}
 
-			for (auto DynamicBindingObject : BPGC->DynamicBindingObjects)
+		for (auto DynamicBindingObject : BPGC->DynamicBindingObjects)
+		{
+			if (DynamicBindingObject)
 			{
-				if (DynamicBindingObject)
-				{
-					HandleClassSubobject(Context, DynamicBindingObject, FEmitterLocalContext::EClassSubobjectList::DynamicBindingObjects, bCreate, bInitilize);
-				}
+				HandleClassSubobject(Context, DynamicBindingObject, FEmitterLocalContext::EClassSubobjectList::DynamicBindingObjects, bCreate, bInitilize);
 			}
-			FBackendHelperUMG::CreateClassSubobjects(Context, bCreate, bInitilize);
-		};
-		CreateAndInitializeClassSubobjects(true, false);
-		CreateAndInitializeClassSubobjects(false, true);
+		}
+		FBackendHelperUMG::CreateClassSubobjects(Context, bCreate, bInitilize);
+	};
+	CreateAndInitializeClassSubobjects(true, false);
+	CreateAndInitializeClassSubobjects(false, true);
 
-		FBackendHelperAnim::CreateAnimClassData(Context);
+	FBackendHelperAnim::CreateAnimClassData(Context);
 
-		Context.DecreaseIndent();
-		Context.AddLine(TEXT("}"));
-	}
+	Context.DecreaseIndent();
+	Context.AddLine(TEXT("}"));
+}
+
+void FEmitDefaultValueHelper::GenerateConstructor(FEmitterLocalContext& Context)
+{
+	GenerateCustomDynamicClassInitialization(Context);
+
+	auto BPGC = CastChecked<UBlueprintGeneratedClass>(Context.GetCurrentlyGeneratedClass());
+	const FString CppClassName = FEmitHelper::GetCppName(BPGC);
+
+	UClass* SuperClass = BPGC->GetSuperClass();
+	const bool bSuperHasOnlyDefaultConstructor = SuperClass && SuperClass->HasMetaData(TEXT("OnlyDefaultConstructorDeclared"));
+
+	Context.CurrentCodeType = FEmitterLocalContext::EGeneratedCodeType::CommonConstructor;
+	Context.ResetPropertiesForInaccessibleStructs();
+	Context.AddLine(FString::Printf(TEXT("%s::%s(const FObjectInitializer& ObjectInitializer) : Super(%s)")
+		, *CppClassName
+		, *CppClassName
+		, bSuperHasOnlyDefaultConstructor ? TEXT("") : TEXT("ObjectInitializer")));
+	Context.AddLine(TEXT("{"));
+	Context.IncreaseIndent();
+
+	// Call CustomDynamicClassInitialization
+	Context.AddLine(FString::Printf(TEXT("if(HasAnyFlags(RF_ClassDefaultObject) && (%s::StaticClass() == GetClass()))"), *CppClassName));
+	Context.AddLine(TEXT("{"));
+	Context.IncreaseIndent();
+	Context.AddLine(FString::Printf(TEXT("%s::__CustomDynamicClassInitialization(CastChecked<UDynamicClass>(GetClass()));"), *CppClassName));
+	Context.DecreaseIndent();
+	Context.AddLine(TEXT("}"));
 
 	// Components that must be fixed after serialization
 	TArray<FString> NativeCreatedComponentProperties;
 
 	{
-		Context.CurrentCodeType = FEmitterLocalContext::EGeneratedCodeType::CommonConstructor;
-		Context.ResetPropertiesForInaccessibleStructs();
-		// Let's have an easy access to generated class subobjects
-		Context.AddLine(TEXT("{")); // no shadow variables
-		Context.IncreaseIndent();
-
 		UObject* CDO = BPGC->GetDefaultObject(false);
 
 		UObject* ParentCDO = BPGC->GetSuperClass()->GetDefaultObject(false);
@@ -927,63 +969,27 @@ void FEmitDefaultValueHelper::GenerateConstructor(FEmitterLocalContext& Context)
 				OuterGenerate(Context, Property, TEXT(""), reinterpret_cast<const uint8*>(CDO), bNewProperty ? nullptr : reinterpret_cast<const uint8*>(ParentCDO), EPropertyAccessOperator::None, true);
 			}
 		}
-
-		Context.DecreaseIndent();
-		Context.AddLine(TEXT("}"));
 	}
 	Context.DecreaseIndent();
 	Context.AddLine(TEXT("}"));
+
 	Context.CurrentCodeType = FEmitterLocalContext::EGeneratedCodeType::Regular;
 	Context.ResetPropertiesForInaccessibleStructs();
+	Context.AddLine(FString::Printf(TEXT("void %s::PostLoadSubobjects(FObjectInstancingGraph* OuterInstanceGraph)"), *CppClassName));
+	Context.AddLine(TEXT("{"));
+	Context.IncreaseIndent();
+	Context.AddLine(TEXT("Super::PostLoadSubobjects(OuterInstanceGraph);"));
+	for (auto& ComponentToFix : NativeCreatedComponentProperties)
 	{
-		Context.AddLine(FString::Printf(TEXT("void %s::PostLoadSubobjects(FObjectInstancingGraph* OuterInstanceGraph)"), *CppClassName));
-		Context.AddLine(TEXT("{"));
-		Context.IncreaseIndent();
-		Context.AddLine(TEXT("Super::PostLoadSubobjects(OuterInstanceGraph);"));
-		
-		for (auto& ComponentToFix : NativeCreatedComponentProperties)
-		{
-			Context.AddLine(FString::Printf(TEXT("if(ensure(%s))"), *ComponentToFix));
-			Context.AddLine(TEXT("{"));
-			Context.IncreaseIndent();
-			Context.AddLine(FString::Printf(TEXT("%s->CreationMethod = EComponentCreationMethod::Native;"), *ComponentToFix));
-			Context.DecreaseIndent();
-			Context.AddLine(TEXT("}"));
-		}
-		Context.DecreaseIndent();
-		Context.AddLine(TEXT("}"));
-	}
-
-	{
-		Context.AddLine(FString::Printf(TEXT("void %s::__GatherReferencedConvertedFields(UDynamicClass* InDynamicClass)"), *CppClassName));
+		Context.AddLine(FString::Printf(TEXT("if(ensure(%s))"), *ComponentToFix));
 		Context.AddLine(TEXT("{"));
 		Context.IncreaseIndent();
-		Context.AddLine(TEXT("ensure(0 == InDynamicClass->ReferencedConvertedFields.Num());"));
-		if (Context.Dependencies.ConvertedClasses.Num())
-		{
-			Context.AddLine(TEXT("// List of all referenced converted classes"));
-		}
-		for (auto LocStruct : Context.Dependencies.ConvertedClasses)
-		{
-			const FString ClassConstructor = FDependenciesHelper::GenerateZConstructor(Context.Dependencies.FindOriginalClass(LocStruct));
-			Context.AddLine(FString::Printf(TEXT("extern UClass* %s;"), *ClassConstructor));
-			Context.AddLine(FString::Printf(TEXT("InDynamicClass->ReferencedConvertedFields.Add(%s);"), *ClassConstructor));
-		}
-
-		if (Context.Dependencies.ConvertedStructs.Num())
-		{
-			Context.AddLine(TEXT("// List of all referenced converted structures"));
-		}
-		for (auto LocStruct : Context.Dependencies.ConvertedStructs)
-		{
-			const FString StructConstructor = FDependenciesHelper::GenerateZConstructor(LocStruct);
-			Context.AddLine(FString::Printf(TEXT("extern UScriptStruct* %s;"), *StructConstructor));
-			Context.AddLine(FString::Printf(TEXT("InDynamicClass->ReferencedConvertedFields.Add(%s);"), *StructConstructor));
-		}
-
+		Context.AddLine(FString::Printf(TEXT("%s->CreationMethod = EComponentCreationMethod::Native;"), *ComponentToFix));
 		Context.DecreaseIndent();
 		Context.AddLine(TEXT("}"));
 	}
+	Context.DecreaseIndent();
+	Context.AddLine(TEXT("}"));
 
 	FDependenciesHelper::AddStaticFunctionsForDependencies(Context);
 
@@ -1012,7 +1018,7 @@ FString FEmitDefaultValueHelper::HandleClassSubobject(FEmitterLocalContext& Cont
 			}
 		}
 
-		const bool AddAsSubobjectOfClass = Object->GetOuter() == Context.GetCurrentlyGeneratedClass();
+		const bool bAddAsSubobjectOfClass = Object->GetOuter() == Context.GetCurrentlyGeneratedClass();
 		LocalNativeName = Context.GenerateUniqueLocalName();
 		Context.AddClassSubObject_InConstructor(Object, LocalNativeName);
 		UClass* ObjectClass = Object->GetClass();
@@ -1022,13 +1028,13 @@ FString FEmitDefaultValueHelper::HandleClassSubobject(FEmitterLocalContext& Cont
 			TEXT("auto %s = NewObject<%s>(%s, %s, TEXT(\"%s\"));")
 			, *LocalNativeName
 			, *NativeType
-			, *OuterStr
+			, bAddAsSubobjectOfClass ? TEXT("InDynamicClass") : *OuterStr
 			, *ActualClass
 			, *Object->GetName()));
-		if (AddAsSubobjectOfClass)
+		if (bAddAsSubobjectOfClass)
 		{
 			Context.RegisterClassSubobject(Object, ListOfSubobjectsType);
-			Context.AddLine(FString::Printf(TEXT("CastChecked<UDynamicClass>(GetClass())->%s.Add(%s);")
+			Context.AddLine(FString::Printf(TEXT("InDynamicClass->%s.Add(%s);")
 				, Context.ClassSubobjectListName(ListOfSubobjectsType)
 				, *LocalNativeName));
 		}
diff --git a/Engine/Source/Developer/BlueprintProfiler/BlueprintProfiler.Build.cs b/Engine/Source/Developer/BlueprintProfiler/BlueprintProfiler.Build.cs
--- a/Engine/Source/Developer/BlueprintProfiler/BlueprintProfiler.Build.cs
+++ b/Engine/Source/Developer/BlueprintProfiler/BlueprintProfiler.Build.cs
@@ -21,8 +21,10 @@ public class BlueprintProfiler : ModuleRules
         {
             PrivateDependencyModuleNames.Add("UnrealEd");
             PrivateDependencyModuleNames.Add("Kismet");
+            PrivateDependencyModuleNames.Add("GraphEditor");
             PrivateDependencyModuleNames.Add("BlueprintGraph");
             PrivateDependencyModuleNames.Add("EditorStyle");
+            PrivateDependencyModuleNames.Add("InputCore");
         }
     }
 }
diff --git a/Engine/Source/Developer/BlueprintProfiler/Private/BlueprintProfiler.cpp b/Engine/Source/Developer/BlueprintProfiler/Private/BlueprintProfiler.cpp
--- a/Engine/Source/Developer/BlueprintProfiler/Private/BlueprintProfiler.cpp
+++ b/Engine/Source/Developer/BlueprintProfiler/Private/BlueprintProfiler.cpp
@@ -151,44 +151,103 @@ bool FBlueprintProfiler::IsTickable() const
 
 void FBlueprintProfiler::ProcessEventProfilingData()
 {
-	int32 StartEvent = 0;
-	int32 StopEvent = 0;
-	TSharedPtr<FBlueprintExecutionContext> BlueprintContext;
+	struct FEventRange
+	{
+		FEventRange()
+			: StartIdx(0)
+			, StopIdx(0)
+		{
+		}
+
+		bool IsRangeValid() const { return StopIdx > StartIdx; }
+
+		TSharedPtr<FBlueprintExecutionContext> BlueprintContext;
+		FName InstanceName;
+		int32 StartIdx;
+		int32 StopIdx;
+	};
+	TArray<FEventRange> ScriptEventRanges;
 	TSet<TSharedPtr<FBlueprintExecutionContext>> DirtyContexts;
 	// Iterate through the events and batch into single script executions
-	for (auto EventIter = InstrumentationEventQueue.CreateConstIterator(); EventIter; ++EventIter)
+	for (int32 EventIdx = 0; EventIdx < InstrumentationEventQueue.Num(); ++EventIdx)
 	{
-		switch (EventIter->GetType())
+		FScriptInstrumentedEvent& CurrEvent = InstrumentationEventQueue[EventIdx];
+		switch (CurrEvent.GetType())
 		{
 			case EScriptInstrumentation::Class:
 			{
-				StartEvent = EventIter.GetIndex();
-				BlueprintContext = GetBlueprintContext(EventIter->GetObjectPath());
+				if ((EventIdx + 2) < InstrumentationEventQueue.Num())
+				{
+					const FScriptInstrumentedEvent& InstanceEvent = InstrumentationEventQueue[EventIdx+1];
+					const FScriptInstrumentedEvent& Event = InstrumentationEventQueue[EventIdx+2];
+					// Check if this is a new event and handle context switch
+					if (InstanceEvent.IsNewInstance() && Event.IsEvent())
+					{
+						FEventRange NewEventRange;
+						NewEventRange.BlueprintContext = GetBlueprintContext(CurrEvent.GetObjectPath());
+						NewEventRange.InstanceName = MapBlueprintInstance(NewEventRange.BlueprintContext, InstanceEvent.GetObjectPath());
+						NewEventRange.StartIdx = EventIdx;
+						ScriptEventRanges.Push(NewEventRange);
+					}
+				}
+				InstrumentationEventQueue.RemoveAt(EventIdx, 2, false);
 				break;
 			}
-			case EScriptInstrumentation::Instance:
+			case EScriptInstrumentation::Event:
 			{
-				MapBlueprintInstance(BlueprintContext, EventIter->GetObjectPath());
+				// Nested events such as calls from event dispatchers.
+				check(ScriptEventRanges.Num() > 0);
+				FEventRange& LastEventRange = ScriptEventRanges.Last();
+				FEventRange NewEventRange;
+				NewEventRange.BlueprintContext = LastEventRange.BlueprintContext;
+				NewEventRange.InstanceName = LastEventRange.InstanceName;
+				NewEventRange.StartIdx = EventIdx;
+				ScriptEventRanges.Push(NewEventRange);
 				break;
 			}
 			case EScriptInstrumentation::Stop:
 			{
-				StopEvent = EventIter.GetIndex();
-				// Create and validate new event.
-				FScriptEventPlayback NewEvent(BlueprintContext);
-				if (NewEvent.Process(InstrumentationEventQueue, StartEvent, StopEvent))
+				const int32 NumEventRanges = ScriptEventRanges.Num();
+				if (ScriptEventRanges.Num() > 0)
 				{
-					DirtyContexts.Add(BlueprintContext);
+					FEventRange& EventRangeToProcess = ScriptEventRanges.Last();
+					EventRangeToProcess.StopIdx = EventIdx;
+					// Create and validate new event.
+					if (EventRangeToProcess.IsRangeValid())
+					{
+						TSharedPtr<FScriptEventPlayback> EventToProcess;
+						// Find any resumed contexts if required
+						if (InstrumentationEventQueue[EventRangeToProcess.StartIdx].IsResumeEvent())
+						{
+							if (TSharedPtr<FScriptEventPlayback>* SuspendedEvent = SuspendedEvents.Find(EventRangeToProcess.InstanceName))
+							{
+								EventToProcess = *SuspendedEvent;
+								SuspendedEvents.Remove(EventRangeToProcess.InstanceName);
+							}
+						}
+						else
+						{
+							EventToProcess = (MakeShareable(new FScriptEventPlayback(EventRangeToProcess.BlueprintContext, EventRangeToProcess.InstanceName)));
+						}
+						check(EventToProcess.IsValid());
+						if (EventToProcess->Process(InstrumentationEventQueue, EventRangeToProcess.StartIdx, EventRangeToProcess.StopIdx))
+						{
+							const int32 NumEventsToRemove = (EventRangeToProcess.StopIdx - EventRangeToProcess.StartIdx) + 1;
+							InstrumentationEventQueue.RemoveAt(EventRangeToProcess.StartIdx, NumEventsToRemove, false);
+							EventIdx -= NumEventsToRemove;
+							DirtyContexts.Add(EventRangeToProcess.BlueprintContext);
+							if (EventToProcess->IsSuspended())
+							{
+								SuspendedEvents.Add(EventRangeToProcess.InstanceName) = EventToProcess;
+							}
+						}
+					}
+					ScriptEventRanges.Pop();
 				}
 				break;
 			}
 		}
 	}
-	// Flush processed event data
-	if (StopEvent > 0)
-	{
-		InstrumentationEventQueue.RemoveAt(0, StopEvent + 1, false);
-	}
 	// Update dirty contexts
 	for (auto Context : DirtyContexts)
 	{
@@ -266,32 +325,24 @@ void FBlueprintProfiler::EndPIE(bool bIsSimulating)
 #endif // WITH_EDITOR
 }
 
-void FBlueprintProfiler::MapBlueprintInstance(TSharedPtr<FBlueprintExecutionContext> BlueprintContext, const FString& InstancePath)
+FName FBlueprintProfiler::MapBlueprintInstance(TSharedPtr<FBlueprintExecutionContext> BlueprintContext, const FString& InstancePath)
 {
 	FName InstanceName(*InstancePath);
 	TWeakObjectPtr<const UObject> Instance = BlueprintContext->GetInstance(InstanceName);
 
-	if (!BlueprintContext->HasProfilerDataForInstance(InstanceName))
+	if (Instance.IsValid() && !BlueprintContext->HasProfilerDataForInstance(InstanceName))
 	{
 		// Create new instance node
-		TSharedPtr<FScriptExecutionInstance> InstanceNode = MakeShareable<FScriptExecutionInstance>(new FScriptExecutionInstance);
-		InstanceNode->SetNameByString(InstanceName.ToString());
-		InstanceNode->SetObservedObject(Instance.Get());
-		InstanceNode->SetGraphName(NAME_None);
-		InstanceNode->SetFlags(EScriptExecutionNodeFlags::Instance);
-		const FSlateBrush* Icon = FEditorStyle::GetBrush(TEXT("BlueprintProfiler.Actor"));
-		InstanceNode->SetIcon(const_cast<FSlateBrush*>(Icon));
-		InstanceNode->SetIconColor(FLinearColor(1.f, 1.f, 1.f, 0.8f));
-		InstanceNode->SetToolTipText(LOCTEXT("NavigateToInstanceHyperlink_ToolTip", "Navigate to the Instance"));
-		// Update widget data
-		if (const AActor* Actor = Cast<AActor>(Instance.Get()))
-		{
-			InstanceNode->SetDisplayName(FText::FromString(Actor->GetActorLabel()));
-		}
-		else
-		{
-			InstanceNode->SetDisplayName(FText::FromString(Instance.Get()->GetName()));
-		}
+		FScriptExecNodeParams InstanceNodeParams;
+		InstanceNodeParams.NodeName = InstanceName;
+		InstanceNodeParams.ObservedObject = Instance.Get();
+		InstanceNodeParams.NodeFlags = EScriptExecutionNodeFlags::Instance;
+		const AActor* Actor = Cast<AActor>(Instance.Get());
+		InstanceNodeParams.DisplayName = Actor ? FText::FromString(Actor->GetActorLabel()) : FText::FromString(Instance.Get()->GetName());
+		InstanceNodeParams.Tooltip = LOCTEXT("NavigateToInstanceHyperlink_ToolTip", "Navigate to the Instance");
+		InstanceNodeParams.IconColor = FLinearColor(1.f, 1.f, 1.f, 0.8f);
+		InstanceNodeParams.Icon = const_cast<FSlateBrush*>(FEditorStyle::GetBrush(TEXT("BlueprintProfiler.Actor")));
+		TSharedPtr<FScriptExecutionInstance> InstanceNode = MakeShareable<FScriptExecutionInstance>(new FScriptExecutionInstance(InstanceNodeParams));
 		// Link to parent blueprint entry
 		TSharedPtr<FScriptExecutionBlueprint> BlueprintNode = BlueprintContext->GetBlueprintExecNode();
 		BlueprintNode->AddInstance(InstanceNode);
@@ -302,6 +353,7 @@ void FBlueprintProfiler::MapBlueprintInstance(TSharedPtr<FBlueprintExecutionCont
 			InstanceNode->AddChildNode(BlueprintNode->GetChildByIndex(NodeIdx));
 		}
 	}
+	return InstanceName;
 }
 
 #undef LOCTEXT_NAMESPACE
diff --git a/Engine/Source/Developer/BlueprintProfiler/Private/ScriptInstrumentationPlayback.cpp b/Engine/Source/Developer/BlueprintProfiler/Private/ScriptInstrumentationPlayback.cpp
--- a/Engine/Source/Developer/BlueprintProfiler/Private/ScriptInstrumentationPlayback.cpp
+++ b/Engine/Source/Developer/BlueprintProfiler/Private/ScriptInstrumentationPlayback.cpp
@@ -2,6 +2,7 @@
 
 #include "BlueprintProfilerPCH.h"
 #include "EditorStyleSet.h"
+#include "GraphEditorSettings.h"
 
 #define LOCTEXT_NAMESPACE "ScriptInstrumentationPlayback"
 
@@ -24,19 +25,19 @@ bool FBlueprintExecutionContext::InitialiseContext(const FString& BlueprintPath)
 			UbergraphFunctionName = BPClass->UberGraphFunction ? BPClass->UberGraphFunction->GetFName() : NAME_None;
 		}
 	}
-	if (Blueprint.IsValid() && BlueprintClass.IsValid())
+	if (Blueprint.IsValid() && BlueprintClass.IsValid() && UbergraphFunctionName != NAME_None)
 	{
 		// Create new blueprint exec node
-		BlueprintNode = MakeShareable<FScriptExecutionBlueprint>(new FScriptExecutionBlueprint);
-		BlueprintNode->SetNameByString(BlueprintPath);
-		BlueprintNode->SetObservedObject(Blueprint.Get());
-		BlueprintNode->SetGraphName(NAME_None);
-		BlueprintNode->SetDisplayName(FText::FromName(Blueprint.Get()->GetFName()));
-		BlueprintNode->SetFlags(EScriptExecutionNodeFlags::Class);
-		const FSlateBrush* Icon_Normal = FEditorStyle::GetBrush(TEXT("BlueprintProfiler.BPIcon_Normal"));
-		BlueprintNode->SetIcon(const_cast<FSlateBrush*>(Icon_Normal));
-		BlueprintNode->SetIconColor(FLinearColor(0.46f, 0.54f, 0.81f));
-		BlueprintNode->SetToolTipText(LOCTEXT("NavigateToBlueprintHyperlink_ToolTip", "Navigate to the Blueprint"));
+		FScriptExecNodeParams BlueprintParams;
+		BlueprintParams.NodeName = FName(*BlueprintPath);
+		BlueprintParams.ObservedObject = Blueprint.Get();
+		BlueprintParams.OwningGraphName = NAME_None;
+		BlueprintParams.DisplayName = FText::FromName(Blueprint.Get()->GetFName());
+		BlueprintParams.Tooltip = LOCTEXT("NavigateToBlueprintHyperlink_ToolTip", "Navigate to the Blueprint");
+		BlueprintParams.NodeFlags = EScriptExecutionNodeFlags::Class;
+		BlueprintParams.Icon = const_cast<FSlateBrush*>(FEditorStyle::GetBrush(TEXT("BlueprintProfiler.BPIcon_Normal")));
+		BlueprintParams.IconColor = FLinearColor(0.46f, 0.54f, 0.81f);
+		BlueprintNode = MakeShareable<FScriptExecutionBlueprint>(new FScriptExecutionBlueprint(BlueprintParams));
 		// Map the blueprint execution
 		if (MapBlueprintExecution())
 		{
@@ -75,6 +76,7 @@ TWeakObjectPtr<const UObject> FBlueprintExecutionContext::GetInstance(FName& Ins
 	if (const FName* SearchName = PIEInstanceNameMap.Find(InstanceNameInOut)) 
 	{
 		CorrectedName = *SearchName;
+		InstanceNameInOut = CorrectedName;
 	}
 	if (TWeakObjectPtr<const UObject>* SearchResult = ActorInstances.Find(CorrectedName))
 	{
@@ -107,22 +109,27 @@ TWeakObjectPtr<const UObject> FBlueprintExecutionContext::GetInstance(FName& Ins
 						}
 					}
 				}
-				else
+				else if (!ObjectPtr->HasAnyFlags(RF_Transient))
 				{
 					ActorInstances.Add(InstanceNameInOut) = ObjectPtr;
+					PIEInstanceNameMap.Add(InstanceNameInOut) = InstanceNameInOut;
 					Result = ObjectPtr;
 				}
 			}
 		}
-
 	}
 	return Result;
 }
 
-TSharedPtr<FBlueprintFunctionContext> FBlueprintExecutionContext::FindFunctionContext(const FName FunctionNameIn) const
+FName FBlueprintExecutionContext::GetEventFunctionName(const FName EventName) const
+{
+	return (EventName == UEdGraphSchema_K2::FN_UserConstructionScript) ? UEdGraphSchema_K2::FN_UserConstructionScript : UbergraphFunctionName;
+}
+
+TSharedPtr<FBlueprintFunctionContext> FBlueprintExecutionContext::GetFunctionContext(const FName FunctionNameIn) const
 {
 	TSharedPtr<FBlueprintFunctionContext> Result;
-	FName FunctionName = (FunctionNameIn ==  UEdGraphSchema_K2::GN_EventGraph) ? UbergraphFunctionName : FunctionNameIn;
+	FName FunctionName = (FunctionNameIn == UEdGraphSchema_K2::GN_EventGraph) ? UbergraphFunctionName : FunctionNameIn;
 
 	if (const TSharedPtr<FBlueprintFunctionContext>* SearchResult = FunctionContexts.Find(FunctionName))
 	{
@@ -131,7 +138,7 @@ TSharedPtr<FBlueprintFunctionContext> FBlueprintExecutionContext::FindFunctionCo
 	return Result;
 }
 
-TSharedPtr<FBlueprintFunctionContext> FBlueprintExecutionContext::GetOrCreateFunctionContext(const FName FunctionNameIn)
+TSharedPtr<FBlueprintFunctionContext> FBlueprintExecutionContext::CreateFunctionContext(const FName FunctionNameIn)
 {
 	FName FunctionName = (FunctionNameIn ==  UEdGraphSchema_K2::GN_EventGraph) ? UbergraphFunctionName : FunctionNameIn;
 	TSharedPtr<FBlueprintFunctionContext>& Result = FunctionContexts.FindOrAdd(FunctionName);
@@ -152,7 +159,7 @@ bool FBlueprintExecutionContext::HasProfilerDataForNode(const UEdGraphNode* Grap
 	const UEdGraph* OuterGraph = GraphNode ? GraphNode->GetTypedOuter<UEdGraph>() : nullptr;
 	if (OuterGraph)
 	{
-		TSharedPtr<FBlueprintFunctionContext> FunctionContext = FindFunctionContext(OuterGraph->GetFName());
+		TSharedPtr<FBlueprintFunctionContext> FunctionContext = GetFunctionContext(OuterGraph->GetFName());
 		if (FunctionContext.IsValid())
 		{
 			return FunctionContext->HasProfilerDataForNode(GraphNode->GetFName());
@@ -168,7 +175,7 @@ TSharedPtr<FScriptExecutionNode> FBlueprintExecutionContext::GetProfilerDataForN
 	const UEdGraph* OuterGraph = GraphNode ? GraphNode->GetTypedOuter<UEdGraph>() : nullptr;
 	if (OuterGraph)
 	{
-		TSharedPtr<FBlueprintFunctionContext> FunctionContext = FindFunctionContext(OuterGraph->GetFName());
+		TSharedPtr<FBlueprintFunctionContext> FunctionContext = GetFunctionContext(OuterGraph->GetFName());
 		if (FunctionContext.IsValid() && FunctionContext->HasProfilerDataForNode(GraphNode->GetFName()))
 		{
 			Result = FunctionContext->GetProfilerDataForNode(GraphNode->GetFName());
@@ -197,6 +204,7 @@ bool FBlueprintExecutionContext::MapBlueprintExecution()
 		TArray<UEdGraph*> Graphs;
 		TArray<UK2Node_Event*> EventNodes;
 		TArray<UK2Node_FunctionEntry*> FunctionEntryNodes;
+		TArray<UK2Node*> InputEventNodes;
 		BlueprintToMap->GetAllGraphs(Graphs);
 		for (auto GraphIter : Graphs)
 		{
@@ -208,56 +216,59 @@ bool FBlueprintExecutionContext::MapBlueprintExecution()
 			{
 				GraphIter->GetNodesOfClass<UK2Node_FunctionEntry>(FunctionEntryNodes);
 			}
+			GraphIter->GetNodesOfClassEx<UK2Node_InputAction, UK2Node>(InputEventNodes);
+			GraphIter->GetNodesOfClassEx<UK2Node_InputKey, UK2Node>(InputEventNodes);
+			GraphIter->GetNodesOfClassEx<UK2Node_InputTouch, UK2Node>(InputEventNodes);
 		}
 		// Map execution paths for each event node
 		for (auto EventIter : EventNodes)
 		{
 			FName EventName = EventIter->GetFunctionName();
-			TSharedPtr<FBlueprintFunctionContext> FunctionContext = GetOrCreateFunctionContext(EventName);
-			TSharedPtr<FScriptExecutionNode> EventExecNode = FunctionContext->GetProfilerDataForNode(EventName);
-			EventExecNode->SetName(EventName);
-			EventExecNode->SetObservedObject(EventIter);
-			EventExecNode->SetGraphName(EventIter->GetTypedOuter<UEdGraph>()->GetFName());
-			EventExecNode->SetFlags(EScriptExecutionNodeFlags::Event);
-			FLinearColor IconColor = FLinearColor(0.91f, 0.16f, 0.16f);
-			FText ToolTipText = LOCTEXT("NavigateToEventLocationHyperlink_ToolTip", "Navigate to the Event");
-			EventExecNode->SetDisplayName(EventIter->GetNodeTitle(ENodeTitleType::ListView));
-			EventExecNode->SetToolTipText(ToolTipText);
-			const FSlateBrush* Icon_Node = EventIter->ShowPaletteIconOnNode() ?	FEditorStyle::GetBrush(EventIter->GetPaletteIcon(IconColor)) :
-																			FEditorStyle::GetBrush(TEXT("BlueprintProfiler.BPNode"));
-			EventExecNode->SetIcon(const_cast<FSlateBrush*>(Icon_Node));
-			EventExecNode->SetIconColor(IconColor);
+			TSharedPtr<FBlueprintFunctionContext> EventContext = CreateFunctionContext(UbergraphFunctionName);
+			FScriptExecNodeParams EventParams;
+			EventParams.NodeName = EventName;
+			EventParams.ObservedObject = EventIter;
+			EventParams.OwningGraphName = NAME_None;
+			EventParams.DisplayName = EventIter->GetNodeTitle(ENodeTitleType::ListView);
+			EventParams.Tooltip = LOCTEXT("NavigateToEventLocationHyperlink_ToolTip", "Navigate to the Event");
+			EventParams.NodeFlags = EScriptExecutionNodeFlags::Event;
+			EventParams.IconColor = FLinearColor(0.91f, 0.16f, 0.16f);
+			const FSlateBrush* EventIcon = EventIter->ShowPaletteIconOnNode() ?	FEditorStyle::GetBrush(EventIter->GetPaletteIcon(EventParams.IconColor)) :
+																				FEditorStyle::GetBrush(TEXT("BlueprintProfiler.BPNode"));
+			EventParams.Icon = const_cast<FSlateBrush*>(EventIcon);
+			TSharedPtr<FScriptExecutionNode> EventExecNode = EventContext->CreateExecutionNode(EventParams);
 			// Add event to mapped blueprint
 			BlueprintNode->AddChildNode(EventExecNode);
 			// Map the event execution flow
-			TSharedPtr<FScriptExecutionNode> MappedEvent = MapNodeExecution(EventIter);
+			TSharedPtr<FScriptExecutionNode> MappedEvent = MapNodeExecution(EventContext, EventIter);
 			if (MappedEvent.IsValid())
 			{
 				EventExecNode->AddChildNode(MappedEvent);
 			}
 		}
+		// Create Input events
+		CreateInputEvents(InputEventNodes);
 		// Map execution paths for each function entry node
 		for (auto FunctionIter : FunctionEntryNodes)
 		{
 			FName FunctionName = FunctionIter->SignatureName;
-			TSharedPtr<FBlueprintFunctionContext> FunctionContext = GetOrCreateFunctionContext(FunctionName);
-			TSharedPtr<FScriptExecutionNode> FunctionExecNode = FunctionContext->GetProfilerDataForNode(FunctionName);
-			FunctionExecNode->SetName(FunctionName);
-			FunctionExecNode->SetObservedObject(FunctionIter);
-			FunctionExecNode->SetGraphName(FunctionIter->GetTypedOuter<UEdGraph>()->GetFName());
-			FunctionExecNode->SetFlags(EScriptExecutionNodeFlags::Event);
-			FLinearColor IconColor = FLinearColor(0.91f, 0.16f, 0.16f);
-			FText ToolTipText = LOCTEXT("NavigateToFunctionLocationHyperlink_ToolTip", "Navigate to the Function");
-			FunctionExecNode->SetDisplayName(FunctionIter->GetNodeTitle(ENodeTitleType::ListView));
-			FunctionExecNode->SetToolTipText(ToolTipText);
-			const FSlateBrush* Icon = FunctionIter->ShowPaletteIconOnNode() ?	FEditorStyle::GetBrush(FunctionIter->GetPaletteIcon(IconColor)) :
-																				FEditorStyle::GetBrush(TEXT("BlueprintProfiler.BPNode"));
-			FunctionExecNode->SetIcon(const_cast<FSlateBrush*>(Icon));
-			FunctionExecNode->SetIconColor(IconColor);
+			TSharedPtr<FBlueprintFunctionContext> FunctionContext = CreateFunctionContext(FunctionName);
+			FScriptExecNodeParams FunctionParams;
+			FunctionParams.NodeName = FunctionName;
+			FunctionParams.ObservedObject = FunctionIter;
+			FunctionParams.OwningGraphName = NAME_None;
+			FunctionParams.DisplayName = FunctionIter->GetNodeTitle(ENodeTitleType::ListView);
+			FunctionParams.Tooltip = LOCTEXT("NavigateToFunctionLocationHyperlink_ToolTip", "Navigate to the Function");
+			FunctionParams.NodeFlags = EScriptExecutionNodeFlags::Event;
+			FunctionParams.IconColor = FLinearColor(0.91f, 0.16f, 0.16f);
+			const FSlateBrush* FunctionIcon = FunctionIter->ShowPaletteIconOnNode() ?	FEditorStyle::GetBrush(FunctionIter->GetPaletteIcon(FunctionParams.IconColor)) :
+																						FEditorStyle::GetBrush(TEXT("BlueprintProfiler.BPNode"));
+			FunctionParams.Icon = const_cast<FSlateBrush*>(FunctionIcon);
+			TSharedPtr<FScriptExecutionNode> FunctionExecNode = FunctionContext->CreateExecutionNode(FunctionParams);
 			// Add Function to mapped blueprint
 			BlueprintNode->AddChildNode(FunctionExecNode);
 			// Map the Function execution flow
-			TSharedPtr<FScriptExecutionNode> MappedFunction = MapNodeExecution(FunctionIter);
+			TSharedPtr<FScriptExecutionNode> MappedFunction = MapNodeExecution(FunctionContext, FunctionIter);
 			if (MappedFunction.IsValid())
 			{
 				FunctionExecNode->AddChildNode(MappedFunction);
@@ -268,39 +279,173 @@ bool FBlueprintExecutionContext::MapBlueprintExecution()
 	return bMappingSuccessful;
 }
 
-TSharedPtr<FScriptExecutionNode> FBlueprintExecutionContext::MapNodeExecution(UEdGraphNode* NodeToMap)
+void FBlueprintExecutionContext::CreateInputEvents(const TArray<UK2Node*>& InputEventNodes)
+{
+	struct FInputEventDesc
+	{
+		UEdGraphNode* GraphNode;
+		UFunction* EventFunction;
+	};
+	TArray<FInputEventDesc> InputEventDescs;
+	// Extract basic event info
+	if (UBlueprintGeneratedClass* BPClass = BlueprintClass.Get())
+	{
+		FInputEventDesc NewInputDesc;
+		for (auto EventNode : InputEventNodes)
+		{
+			if (EventNode)
+			{
+				if (UK2Node_InputAction* InputActionNode = Cast<UK2Node_InputAction>(EventNode))
+				{
+					const FString CustomEventName = FString::Printf(TEXT("InpActEvt_%s_%s"), *InputActionNode->InputActionName.ToString());
+					for (auto FunctionIter = TFieldIterator<UFunction>(BPClass); FunctionIter; ++FunctionIter)
+					{
+						if (FunctionIter->GetName().Contains(CustomEventName))
+						{
+							NewInputDesc.GraphNode = EventNode;
+							NewInputDesc.EventFunction = *FunctionIter;
+							InputEventDescs.Add(NewInputDesc);
+						}
+					}
+				}
+				else if (UK2Node_InputKey* InputKeyNode = Cast<UK2Node_InputKey>(EventNode))
+				{
+					const FName ModifierName = InputKeyNode->GetModifierName();
+					const FString CustomEventName = ModifierName != NAME_None ? FString::Printf(TEXT("InpActEvt_%s_%s"), *ModifierName.ToString(), *InputKeyNode->InputKey.ToString()) :
+																				FString::Printf(TEXT("InpActEvt_%s"), *InputKeyNode->InputKey.ToString());
+					for (auto FunctionIter = TFieldIterator<UFunction>(BPClass); FunctionIter; ++FunctionIter)
+					{
+						if (FunctionIter->GetName().Contains(CustomEventName))
+						{
+							NewInputDesc.GraphNode = EventNode;
+							NewInputDesc.EventFunction = *FunctionIter;
+							InputEventDescs.Add(NewInputDesc);
+						}
+					}
+				}
+				else if (UK2Node_InputTouch* InputTouchNode = Cast<UK2Node_InputTouch>(EventNode))
+				{
+					struct EventPinData
+					{
+						EventPinData(UEdGraphPin* InPin,TEnumAsByte<EInputEvent> InEvent ){	Pin=InPin;EventType=InEvent;};
+						UEdGraphPin* Pin;
+						TEnumAsByte<EInputEvent> EventType;
+					};
+					TArray<UEdGraphPin*> ActivePins;
+					if (UEdGraphPin* InputTouchPressedPin = InputTouchNode->GetPressedPin())
+					{
+						if (InputTouchPressedPin->LinkedTo.Num() > 0)
+						{
+							ActivePins.Add(InputTouchPressedPin);
+						}
+					}
+					if (UEdGraphPin* InputTouchReleasedPin = InputTouchNode->GetReleasedPin())
+					{
+						if (InputTouchReleasedPin->LinkedTo.Num() > 0)
+						{
+							ActivePins.Add(InputTouchReleasedPin);
+						}
+					}
+					if (UEdGraphPin* InputTouchMovedPin = InputTouchNode->GetMovedPin())
+					{
+						if (InputTouchMovedPin->LinkedTo.Num() > 0)
+						{
+							ActivePins.Add(InputTouchMovedPin);
+						}
+					}
+					for (auto Pin : ActivePins)
+					{
+						const FString CustomEventName = FString::Printf(TEXT("InpTchEvt_%s"), *Pin->GetName());
+						for (auto FunctionIter = TFieldIterator<UFunction>(BPClass); FunctionIter; ++FunctionIter)
+						{
+							if (FunctionIter->GetName().Contains(CustomEventName))
+							{
+								NewInputDesc.GraphNode = EventNode;
+								NewInputDesc.EventFunction = *FunctionIter;
+								InputEventDescs.Add(NewInputDesc);
+							}
+						}
+					}
+				}
+			}
+		}
+	}
+	// Build event contexts
+	for (auto InputEventDesc : InputEventDescs)
+	{
+		UEdGraphNode* GraphNode = InputEventDesc.GraphNode;
+		FName EventName = InputEventDesc.EventFunction->GetFName();
+		TSharedPtr<FBlueprintFunctionContext> EventContext = CreateFunctionContext(UbergraphFunctionName);
+		FScriptExecNodeParams EventParams;
+		EventParams.NodeName = EventName;
+		EventParams.ObservedObject = GraphNode;
+		EventParams.OwningGraphName = NAME_None;
+		EventParams.DisplayName = GraphNode->GetNodeTitle(ENodeTitleType::ListView);
+		EventParams.Tooltip = LOCTEXT("NavigateToEventLocationHyperlink_ToolTip", "Navigate to the Event");
+		EventParams.NodeFlags = EScriptExecutionNodeFlags::Event;
+		EventParams.IconColor = FLinearColor(0.91f, 0.16f, 0.16f);
+		const FSlateBrush* EventIcon = GraphNode->ShowPaletteIconOnNode() ?	FEditorStyle::GetBrush(GraphNode->GetPaletteIcon(EventParams.IconColor)) :
+																			FEditorStyle::GetBrush(TEXT("BlueprintProfiler.BPNode"));
+		EventParams.Icon = const_cast<FSlateBrush*>(EventIcon);
+		TSharedPtr<FScriptExecutionNode> EventExecNode = EventContext->CreateExecutionNode(EventParams);
+		// Add Function to mapped blueprint
+		BlueprintNode->AddChildNode(EventExecNode);
+		// Map the Function execution flow
+		TSharedPtr<FScriptExecutionNode> MappedFunction = MapNodeExecution(EventContext, GraphNode);
+		if (MappedFunction.IsValid())
+		{
+			EventExecNode->AddChildNode(MappedFunction);
+		}
+	}
+}
+
+TSharedPtr<FScriptExecutionNode> FBlueprintExecutionContext::MapNodeExecution(TSharedPtr<FBlueprintFunctionContext> FunctionContext, UEdGraphNode* NodeToMap)
 {
 	TSharedPtr<FScriptExecutionNode> MappedNode;
 	if (NodeToMap)
 	{
-		const UEdGraph* OwningGraph = NodeToMap->GetTypedOuter<UEdGraph>();
-		TSharedPtr<FBlueprintFunctionContext> FunctionContext = GetOrCreateFunctionContext(OwningGraph->GetFName());
-		const bool bNodeMapped = FunctionContext->HasProfilerDataForNode(NodeToMap->GetFName());
 		MappedNode = FunctionContext->GetProfilerDataForNode(NodeToMap->GetFName());
-
-		if (!bNodeMapped)
+		if (!MappedNode.IsValid())
 		{
-			// Construct mapped node values
-			MappedNode->SetName(NodeToMap->GetFName());
-			MappedNode->SetObservedObject(NodeToMap);
-			MappedNode->SetGraphName(NodeToMap->GetTypedOuter<UEdGraph>()->GetFName());
-			MappedNode->SetFlags(EScriptExecutionNodeFlags::Node);
-			FLinearColor IconColor = FLinearColor(1.f, 1.f, 1.f, 0.8f);
-			FText ToolTipText = LOCTEXT("NavigateToNodeLocationHyperlink_ToolTip", "Navigate to the Node");
-			// Evaluate Execution Pins
+			FScriptExecNodeParams NodeParams;
+			NodeParams.NodeName = NodeToMap->GetFName();
+			NodeParams.ObservedObject = NodeToMap;
+			NodeParams.DisplayName = NodeToMap->GetNodeTitle(ENodeTitleType::ListView);
+			NodeParams.Tooltip = LOCTEXT("NavigateToNodeLocationHyperlink_ToolTip", "Navigate to the Node");
+			NodeParams.NodeFlags = EScriptExecutionNodeFlags::Node;
+			NodeParams.IconColor = FLinearColor(1.f, 1.f, 1.f, 0.8f);
+			const FSlateBrush* NodeIcon = NodeToMap->ShowPaletteIconOnNode() ?	FEditorStyle::GetBrush(NodeToMap->GetPaletteIcon(NodeParams.IconColor)) :
+																				FEditorStyle::GetBrush(TEXT("BlueprintProfiler.BPNode"));
+			NodeParams.Icon = const_cast<FSlateBrush*>(NodeIcon);
+			MappedNode = FunctionContext->CreateExecutionNode(NodeParams);
+			// Evaluate Execution and Input Pins
 			TArray<UEdGraphPin*> ExecPins;
+			TArray<UEdGraphPin*> InputPins;
 			for (auto Pin : NodeToMap->Pins)
 			{
 				if (Pin->Direction == EGPD_Output && Pin->PinType.PinCategory == UEdGraphSchema_K2::PC_Exec)
 				{
 					ExecPins.Add(Pin);
 				}
+				else if (Pin->Direction == EGPD_Input && Pin->PinType.PinCategory != UEdGraphSchema_K2::PC_Exec)
+				{
+					InputPins.Add(Pin);
+				}
+			}
+
+			if (ExecPins.Num() == 0)
+			{
+				MappedNode->AddFlags(EScriptExecutionNodeFlags::PureNode);
 			}
+			else
+			{
+				MappedNode->SetPureNodeScriptCodeRange(FunctionContext->GetPureNodeScriptCodeRange(NodeToMap));
+			}
+
 			const bool bBranchNode = ExecPins.Num() > 1;
 			if (bBranchNode)
 			{
 				// Determine branch exec type
-				IconColor = FLinearColor(1.f, 1.f, 1.f, 0.8f);
 				if (NodeToMap->IsA<UK2Node_ExecutionSequence>())
 				{
 					MappedNode->AddFlags(EScriptExecutionNodeFlags::SequentialBranch);
@@ -310,6 +455,31 @@ TSharedPtr<FScriptExecutionNode> FBlueprintExecutionContext::MapNodeExecution(UE
 					MappedNode->AddFlags(EScriptExecutionNodeFlags::ConditionalBranch);
 				}
 			}
+			// Evaluate non-exec input pins (pure node execution chains)
+			TArray<int32> PinScriptCodeOffsets;
+			for (auto InputPin : InputPins)
+			{
+				// If this input pin is linked to a pure node in the source graph, create and map all known execution paths for it.
+				for (auto LinkedPin : InputPin->LinkedTo)
+				{
+					// Note: Intermediate pure nodes can have output pins that masquerade as impure node output pins when links are "moved" from the source graph (thus
+					// resulting in a false association here with one or more script code offsets), so we must first ensure that the link is really to a pure node output.
+					UK2Node* OwningNode = Cast<UK2Node>(LinkedPin->GetOwningNode());
+					if (OwningNode && OwningNode->IsNodePure())
+					{
+						FunctionContext->GetAllCodeLocationsFromPin(LinkedPin, PinScriptCodeOffsets);
+						if (PinScriptCodeOffsets.Num() > 0)
+						{
+							TSharedPtr<FScriptExecutionNode> PureNode = MapNodeExecution(FunctionContext, OwningNode);
+							check(PureNode.IsValid() && PureNode->IsPureNode());
+							for (int32 i = 0; i < PinScriptCodeOffsets.Num(); ++i)
+							{
+								MappedNode->AddPureNode(PinScriptCodeOffsets[i], PureNode);
+							}
+						}
+					}
+				}
+			}
 			for (auto Pin : ExecPins)
 			{
 				const UEdGraphPin* LinkedPin = Pin->LinkedTo.Num() > 0 ? Pin->LinkedTo[0] : nullptr;
@@ -319,22 +489,22 @@ TSharedPtr<FScriptExecutionNode> FBlueprintExecutionContext::MapNodeExecution(UE
 					const int32 PinScriptCodeOffset = FunctionContext->GetCodeLocationFromPin(Pin);
 					// Check we have a valid script offset or that the pin isn't linked.
 					check((PinScriptCodeOffset != INDEX_NONE)||(LinkedPin == nullptr));
-					TSharedPtr<FScriptExecutionNode> LinkedNode = FunctionContext->GetProfilerDataForNode(Pin->GetFName());
-					LinkedNode->SetName(Pin->GetFName());
-					LinkedNode->SetObservedObject(Pin);
-					LinkedNode->SetGraphName(Pin->GetTypedOuter<UEdGraph>()->GetFName());
-					LinkedNode->SetDisplayName(Pin->GetDisplayName());
-					LinkedNode->SetFlags(EScriptExecutionNodeFlags::ExecPin);
-					LinkedNode->SetToolTipText(LOCTEXT("ExecPin_ToolTip", "Expand execution path"));
+					FScriptExecNodeParams LinkNodeParams;
+					LinkNodeParams.NodeName = Pin->GetFName();
+					LinkNodeParams.ObservedObject = Pin;
+					LinkNodeParams.DisplayName = Pin->GetDisplayName();
+					LinkNodeParams.Tooltip = LOCTEXT("ExecPin_ToolTip", "Expand execution path");
+					LinkNodeParams.NodeFlags = EScriptExecutionNodeFlags::ExecPin;
+					LinkNodeParams.IconColor = FLinearColor(1.f, 1.f, 1.f, 0.8f);
 					const FSlateBrush* Icon = LinkedPin ?	FEditorStyle::GetBrush(TEXT("BlueprintProfiler.BPPinConnected")) : 
 															FEditorStyle::GetBrush(TEXT("BlueprintProfiler.BPPinDisconnected"));
-					LinkedNode->SetIcon(const_cast<FSlateBrush*>(Icon));
-					LinkedNode->SetIconColor(FLinearColor(1.f, 1.f, 1.f, 0.8f));
+					LinkNodeParams.Icon = const_cast<FSlateBrush*>(Icon);
+					TSharedPtr<FScriptExecutionNode> LinkedNode = FunctionContext->CreateExecutionNode(LinkNodeParams);
 					MappedNode->AddLinkedNode(PinScriptCodeOffset, LinkedNode);
 					// Map pin exec flow
 					if (LinkedPin)
 					{
-						TSharedPtr<FScriptExecutionNode> ExecPinNode = MapNodeExecution(LinkedPin->GetOwningNode());
+						TSharedPtr<FScriptExecutionNode> ExecPinNode = MapNodeExecution(FunctionContext, LinkedPin->GetOwningNode());
 						if (ExecPinNode.IsValid())
 						{
 							LinkedNode->AddChildNode(ExecPinNode);
@@ -343,7 +513,7 @@ TSharedPtr<FScriptExecutionNode> FBlueprintExecutionContext::MapNodeExecution(UE
 				}
 				else if (LinkedPin)
 				{
-					TSharedPtr<FScriptExecutionNode> ExecLinkedNode = MapNodeExecution(LinkedPin->GetOwningNode());
+					TSharedPtr<FScriptExecutionNode> ExecLinkedNode = MapNodeExecution(FunctionContext, LinkedPin->GetOwningNode());
 					if (ExecLinkedNode.IsValid())
 					{
 						const int32 PinScriptCodeOffset = FunctionContext->GetCodeLocationFromPin(Pin);
@@ -360,19 +530,25 @@ TSharedPtr<FScriptExecutionNode> FBlueprintExecutionContext::MapNodeExecution(UE
 				{
 					// Update Exec node
 					MappedNode->AddFlags(EScriptExecutionNodeFlags::FunctionCall);
-					IconColor = FLinearColor(0.46f, 0.54f, 0.95f);
-					ToolTipText = LOCTEXT("NavigateToFunctionLocationHyperlink_ToolTip", "Navigate to the Function Callsite");
+					MappedNode->SetIconColor(FLinearColor(0.46f, 0.54f, 0.95f));
+					MappedNode->SetToolTipText(LOCTEXT("NavigateToFunctionLocationHyperlink_ToolTip", "Navigate to the Function Callsite"));
+					// Update the function context
+					FunctionContext = CreateFunctionContext(CalledGraph->GetFName());
 					// Grab function entries
 					TArray<UK2Node_FunctionEntry*> FunctionEntryNodes;
 					CalledGraph->GetNodesOfClass<UK2Node_FunctionEntry>(FunctionEntryNodes);
 	
 					for (auto EntryNodeIter : FunctionEntryNodes)
 					{
-						TSharedPtr<FScriptExecutionNode> FunctionEntry = MapNodeExecution(EntryNodeIter);
+						TSharedPtr<FScriptExecutionNode> FunctionEntry = MapNodeExecution(FunctionContext, EntryNodeIter);
 						MappedNode->AddChildNode(FunctionEntry);
 					}
 				}
 			}
+			else if (UK2Node_CustomEvent* CustomEvent = Cast<UK2Node_CustomEvent>(NodeToMap))
+			{
+				MappedNode->AddFlags(EScriptExecutionNodeFlags::CustomEvent);
+			}
 			//else if (UK2Node_MacroInstance* MacroCallSite = Cast<UK2Node_MacroInstance>(NodeToMap))
 			//{
 			//	if (UEdGraph* MacroGraph = MacroCallSite->GetMacroGraph())
@@ -396,13 +572,6 @@ TSharedPtr<FScriptExecutionNode> FBlueprintExecutionContext::MapNodeExecution(UE
 			//		}
 			//	}
 			//}
-			// Finalise exec node data.
-			MappedNode->SetDisplayName(NodeToMap->GetNodeTitle(ENodeTitleType::ListView));
-			MappedNode->SetToolTipText(ToolTipText);
-			const FSlateBrush* Icon = NodeToMap->ShowPaletteIconOnNode() ?	FEditorStyle::GetBrush(NodeToMap->GetPaletteIcon(IconColor)) :
-																			FEditorStyle::GetBrush(TEXT("BlueprintProfiler.BPNode"));
-			MappedNode->SetIcon(const_cast<FSlateBrush*>(Icon));
-			MappedNode->SetIconColor(IconColor);
 		}
 	}
 	return MappedNode;
@@ -417,10 +586,10 @@ bool FBlueprintFunctionContext::HasProfilerDataForNode(const FName NodeName) con
 
 TSharedPtr<FScriptExecutionNode> FBlueprintFunctionContext::GetProfilerDataForNode(const FName NodeName)
 {
-	TSharedPtr<FScriptExecutionNode>& Result = ExecutionNodes.FindOrAdd(NodeName);
-	if (!Result.IsValid())
+	TSharedPtr<FScriptExecutionNode> Result;
+	if (TSharedPtr<FScriptExecutionNode>* SearchResult = ExecutionNodes.Find(NodeName))
 	{
-		Result = MakeShareable(new FScriptExecutionNode);
+		Result = *SearchResult;
 	}
 	return Result;
 }
@@ -438,63 +607,99 @@ const UEdGraphNode* FBlueprintFunctionContext::GetNodeFromCodeLocation(const int
 	return Result.Get();
 }
 
+const UEdGraphPin* FBlueprintFunctionContext::GetPinFromCodeLocation(const int32 ScriptOffset)
+{
+	TWeakObjectPtr<const UEdGraphPin>& Result = ScriptOffsetToPins.FindOrAdd(ScriptOffset);
+	if (!Result.IsValid() && BlueprintClass.IsValid())
+	{
+		if (const UEdGraphPin* GraphPin = BlueprintClass->GetDebugData().FindSourcePinFromCodeLocation(Function.Get(), ScriptOffset))
+		{
+			Result = GraphPin;
+		}
+	}
+	return Result.Get();
+}
+
 const int32 FBlueprintFunctionContext::GetCodeLocationFromPin(const UEdGraphPin* Pin) const
 {
-	int32 Result = INDEX_NONE;
 	if (BlueprintClass.IsValid())
 	{
-		Result = BlueprintClass.Get()->GetDebugData().FindPinToCodeAssociation(Pin, Function.Get());
+		return BlueprintClass.Get()->GetDebugData().FindCodeLocationFromSourcePin(Pin, Function.Get());
+	}
+
+	return INDEX_NONE;
+}
+
+void FBlueprintFunctionContext::GetAllCodeLocationsFromPin(const UEdGraphPin* Pin, TArray<int32>& OutCodeLocations) const
+{
+	if (BlueprintClass.IsValid())
+	{
+		BlueprintClass.Get()->GetDebugData().FindAllCodeLocationsFromSourcePin(Pin, Function.Get(), OutCodeLocations);
 	}
-	return Result;
+}
+
+FInt32Range FBlueprintFunctionContext::GetPureNodeScriptCodeRange(const UEdGraphNode* Node) const
+{
+	if (BlueprintClass.IsValid())
+	{
+		return BlueprintClass->GetDebugData().FindPureNodeScriptCodeRangeFromSourceNode(Node, Function.Get());
+	}
+
+	return FInt32Range(INDEX_NONE);
+}
+
+TSharedPtr<FScriptExecutionNode> FBlueprintFunctionContext::CreateExecutionNode(FScriptExecNodeParams& InitParams)
+{
+	TSharedPtr<FScriptExecutionNode>& ScriptExecNode = ExecutionNodes.FindOrAdd(InitParams.NodeName);
+	check(!ScriptExecNode.IsValid());
+	UEdGraph* OuterGraph = InitParams.ObservedObject.IsValid() ? InitParams.ObservedObject.Get()->GetTypedOuter<UEdGraph>() : nullptr;
+	InitParams.OwningGraphName = OuterGraph ? OuterGraph->GetFName() : NAME_None;
+	ScriptExecNode = MakeShareable(new FScriptExecutionNode(InitParams));
+	return ScriptExecNode;
 }
 
 //////////////////////////////////////////////////////////////////////////
 // FScriptEventPlayback
 
-bool FScriptEventPlayback::Process(const TArray<FScriptInstrumentedEvent>& Events, const int32 Start, const int32 Stop)
+bool FScriptEventPlayback::Process(const TArray<FScriptInstrumentedEvent>& SignalData, const int32 StartIdx, const int32 StopIdx)
 {
 	struct PerNodeEventHelper
 	{
 		TArray<FTracePath> EntryTracePaths;
+		TArray<FTracePath> InputTracePaths;
 		TSharedPtr<FScriptExecutionNode> ExecNode;
 		TSharedPtr<FBlueprintFunctionContext> FunctionContext;
 		TArray<FScriptInstrumentedEvent> UnproccessedEvents;
 	};
 
-	const int32 NumEvents = (Stop + 1) - Start;
-	if (BlueprintContext.IsValid() && NumEvents > 2)
+	const int32 NumEvents = (StopIdx+1) - StartIdx;
+	ProcessingState = EEventProcessingResult::Failed;
+	const bool bEventIsResuming = SignalData[StartIdx].IsResumeEvent() ;
+
+	if (BlueprintContext.IsValid() && InstanceName != NAME_None)
 	{
-		check(Events[Start].GetType() == EScriptInstrumentation::Class);
-		check(Events[Start+1].GetType() == EScriptInstrumentation::Instance);
-		check(Events[Start+2].GetType() == EScriptInstrumentation::Event);
-		// Update event and instance names
-		EventName = Events[Start+2].GetFunctionName();
-		InstanceName = BlueprintContext->RemapInstancePath(FName(*Events[Start+1].GetObjectPath()));
-		// Process event
-		FName CurrentFunctionName = EventName;
-		TSharedPtr<FBlueprintFunctionContext> FunctionContext = BlueprintContext->GetOrCreateFunctionContext(CurrentFunctionName);
+		check(SignalData[StartIdx].IsEvent());
+		EventName = bEventIsResuming ? EventName : SignalData[StartIdx].GetFunctionName();
+		FName CurrentFunctionName = BlueprintContext->GetEventFunctionName(SignalData[StartIdx].GetFunctionName());
+		TSharedPtr<FBlueprintFunctionContext> FunctionContext = BlueprintContext->GetFunctionContext(CurrentFunctionName);
 		TSharedPtr<FScriptExecutionNode> EventNode = FunctionContext->GetProfilerDataForNode(EventName);
-		FTracePath EventTrace;
-		TSharedPtr<FScriptPerfData> PerfData = EventNode->GetPerfDataByInstanceAndTracePath(InstanceName, EventTrace);
-		PerfData->AddEventTiming(0.0, Events[Stop].GetTime() - Events[Start+2].GetTime());
 		// Find Associated graph nodes and submit into a node map for later processing.
-		FTracePath TracePath;
-		TArray<FTracePath> TraceStack;
 		TMap<const UEdGraphNode*, PerNodeEventHelper> NodeMap;
+		ProcessingState = EEventProcessingResult::Success;
+		int32 LastEventIdx = StartIdx;
+		const int32 EventStartOffset = SignalData[StartIdx].IsResumeEvent() ? 3 : 1;
 
-		for (int32 EventIdx = Start + 3; EventIdx < Stop; ++EventIdx)
+		for (int32 SignalIdx = StartIdx + EventStartOffset; SignalIdx < StopIdx; ++SignalIdx)
 		{
-			const FScriptInstrumentedEvent& CurrEvent = Events[EventIdx];
-			bool bFunctionChanged = false;
+			const FScriptInstrumentedEvent& CurrSignal = SignalData[SignalIdx];
 			// Update script function.
-			if (CurrentFunctionName != CurrEvent.GetFunctionName())
+			if (CurrentFunctionName != CurrSignal.GetFunctionName())
 			{
-				bFunctionChanged = true;
-				CurrentFunctionName = CurrEvent.GetFunctionName();
-				FunctionContext = BlueprintContext->FindFunctionContext(CurrentFunctionName);
+				CurrentFunctionName = CurrSignal.GetFunctionName();
+				FunctionContext = BlueprintContext->GetFunctionContext(CurrentFunctionName);
 				check(FunctionContext.IsValid());
 			}
-			if (const UEdGraphNode* GraphNode = FunctionContext->GetNodeFromCodeLocation(CurrEvent.GetScriptCodeOffset()))
+			if (const UEdGraphNode* GraphNode = FunctionContext->GetNodeFromCodeLocation(CurrSignal.GetScriptCodeOffset()))
 			{
 				PerNodeEventHelper& NodeEntry = NodeMap.FindOrAdd(GraphNode);
 				if (!NodeEntry.FunctionContext.IsValid())
@@ -502,27 +707,48 @@ bool FScriptEventPlayback::Process(const TArray<FScriptInstrumentedEvent>& Event
 					NodeEntry.FunctionContext = FunctionContext;
 					NodeEntry.ExecNode = FunctionContext->GetProfilerDataForNode(GraphNode->GetFName());
 				}
-				switch (CurrEvent.GetType())
+				switch (CurrSignal.GetType())
 				{
 					case EScriptInstrumentation::PureNodeEntry:
 					{
-						NodeEntry.UnproccessedEvents.Add(CurrEvent);
+						if (const UEdGraphPin* Pin = FunctionContext->GetPinFromCodeLocation(CurrSignal.GetScriptCodeOffset()))
+						{
+							if (NodeEntry.FunctionContext->HasProfilerDataForNode(Pin->GetOwningNode()->GetFName()))
+							{
+								TracePath.AddExitPin(CurrSignal.GetScriptCodeOffset());
+								NodeEntry.InputTracePaths.Insert(TracePath, 0);
+							}
+						}
+						NodeEntry.UnproccessedEvents.Add(CurrSignal);
 						break;
 					}
 					case EScriptInstrumentation::NodeEntry:
 					{
-						NodeEntry.EntryTracePaths.Push(TracePath);
-						if (bFunctionChanged && GraphNode->IsA<UK2Node_FunctionEntry>())
+						// Handle timings for events called as functions
+						if (NodeEntry.ExecNode->IsCustomEvent())
 						{
-							TraceStack.Push(TracePath);
+							// Ensure this is a different event
+							if (const UK2Node_Event* EventGraphNode = Cast<UK2Node_Event>(GraphNode))
+							{
+								const FName NewEventName = EventGraphNode->GetFunctionName();
+								if (NewEventName != EventName)
+								{
+									TracePath.Reset();
+									EventTimings.Add(EventName) = SignalData[SignalIdx].GetTime() - SignalData[LastEventIdx].GetTime();
+									EventName = NewEventName;
+									EventNode = FunctionContext->GetProfilerDataForNode(EventName);
+									LastEventIdx = SignalIdx;
+								}
+							}
 						}
-						NodeEntry.UnproccessedEvents.Add(CurrEvent);
+						NodeEntry.EntryTracePaths.Push(TracePath);
+						NodeEntry.UnproccessedEvents.Add(CurrSignal);
 						break;
 					}
 					case EScriptInstrumentation::PushState:
 					{
 						TraceStack.Push(TracePath);
-						NodeEntry.UnproccessedEvents.Add(CurrEvent);
+						NodeEntry.UnproccessedEvents.Add(CurrSignal);
 						break;
 					}
 					case EScriptInstrumentation::RestoreState:
@@ -531,44 +757,48 @@ bool FScriptEventPlayback::Process(const TArray<FScriptInstrumentedEvent>& Event
 						TracePath = TraceStack.Last();
 						break;
 					}
+					case EScriptInstrumentation::SuspendState:
+					{
+						ProcessingState = EEventProcessingResult::Suspended;
+						break;
+					}
 					case EScriptInstrumentation::PopState:
 					{
 						if (TraceStack.Num())
 						{
 							TracePath = TraceStack.Pop();
 						}
-						FScriptInstrumentedEvent OverrideEvent(CurrEvent);
-						OverrideEvent.OverrideType(EScriptInstrumentation::NodeExit);
-						NodeEntry.UnproccessedEvents.Add(OverrideEvent);
-						// Consolidate enclosed timings
-						int32 NodeEntryIdx = INDEX_NONE;
-						for (int32 NodeEventIdx = NodeEntry.UnproccessedEvents.Num()-2; NodeEventIdx >= 0; --NodeEventIdx)
+						// Consolidate enclosed timings for execution sequence nodes
+						if (GraphNode->IsA<UK2Node_ExecutionSequence>())
 						{
-							const FScriptInstrumentedEvent& NodeEvent = NodeEntry.UnproccessedEvents[NodeEventIdx];
-
-							if (NodeEvent.GetType() == EScriptInstrumentation::NodeEntry)
+							FScriptInstrumentedEvent OverrideEvent(CurrSignal);
+							OverrideEvent.OverrideType(EScriptInstrumentation::NodeExit);
+							NodeEntry.UnproccessedEvents.Add(OverrideEvent);
+							int32 NodeEntryIdx = INDEX_NONE;
+							for (int32 NodeEventIdx = NodeEntry.UnproccessedEvents.Num()-2; NodeEventIdx >= 0; --NodeEventIdx)
 							{
-								if (NodeEntryIdx != INDEX_NONE)
+								const FScriptInstrumentedEvent& NodeEvent = NodeEntry.UnproccessedEvents[NodeEventIdx];
+
+								if (NodeEvent.GetType() == EScriptInstrumentation::NodeEntry)
 								{
-									// Keep only the first node entry.
+									if (NodeEntryIdx != INDEX_NONE)
+									{
+										// Keep only the first node entry.
+										NodeEntry.UnproccessedEvents.RemoveAt(NodeEventIdx, 1);
+									}
+									NodeEntryIdx = NodeEventIdx;
+								}
+								else if (NodeEvent.GetType() == EScriptInstrumentation::NodeExit)
+								{
+									// Strip out all node exits apart from the PopExecution overridden event.
 									NodeEntry.UnproccessedEvents.RemoveAt(NodeEventIdx, 1);
 								}
-								NodeEntryIdx = NodeEventIdx;
-							}
-							else if (NodeEvent.GetType() == EScriptInstrumentation::NodeExit)
-							{
-								// Strip out all node exits apart from the PopExecution overridden event.
-								NodeEntry.UnproccessedEvents.RemoveAt(NodeEventIdx, 1);
 							}
 						}
 						break;
 					}
 					case EScriptInstrumentation::NodeExit:
 					{
-						if (TraceStack.Num() && bFunctionChanged)
-						{
-							TracePath = TraceStack.Pop();
-						}
 						// Cleanup branching multiple exits and correct the tracepath
 						if (NodeEntry.UnproccessedEvents.Num() && NodeEntry.UnproccessedEvents.Last().GetType() == EScriptInstrumentation::NodeExit)
 						{
@@ -578,34 +808,80 @@ bool FScriptEventPlayback::Process(const TArray<FScriptInstrumentedEvent>& Event
 								TracePath = NodeEntry.EntryTracePaths.Last();
 							}
 						}
-						TracePath.AddExitPin(CurrEvent.GetScriptCodeOffset());
-						NodeEntry.UnproccessedEvents.Add(CurrEvent);
+						TracePath.AddExitPin(CurrSignal.GetScriptCodeOffset());
+						NodeEntry.UnproccessedEvents.Add(CurrSignal);
 						break;
 					}
 					default:
 					{
-						NodeEntry.UnproccessedEvents.Add(CurrEvent);
+						NodeEntry.UnproccessedEvents.Add(CurrSignal);
 						break;
 					}
 				}
 			}
 		}
-		// Process Node map timings
+		// Process last event timing
+		double* TimingData = bEventIsResuming ? EventTimings.Find(EventName) : nullptr;
+		if (TimingData)
+		{
+			*TimingData += SignalData[StopIdx].GetTime() - SignalData[LastEventIdx].GetTime();
+		}
+		else
+		{
+			EventTimings.Add(EventName) = SignalData[StopIdx].GetTime() - SignalData[LastEventIdx].GetTime();
+		}
+		// Process outstanding event timings, adding to previous timings if existing.
+		if (ProcessingState != EEventProcessingResult::Suspended)
+		{
+			for (auto EventTiming : EventTimings)
+			{
+				FTracePath EventTracePath;
+				const FName EventFunctionName = BlueprintContext->GetEventFunctionName(EventTiming.Key);
+				FunctionContext = BlueprintContext->GetFunctionContext(EventFunctionName);
+				check (FunctionContext.IsValid());
+				EventNode = FunctionContext->GetProfilerDataForNode(EventTiming.Key);
+				check (EventNode.IsValid());
+				TSharedPtr<FScriptPerfData> PerfData = EventNode->GetPerfDataByInstanceAndTracePath(InstanceName, EventTracePath);
+				PerfData->AddEventTiming(0.0, SignalData[StopIdx].GetTime() - SignalData[LastEventIdx].GetTime());
+			}
+			EventTimings.Reset();
+		}
+		// Process Node map timings -- this can probably be rolled back into submission during the initial processing and lose this extra iteration.
 		for (auto NodeEntry : NodeMap)
 		{
 			TSharedPtr<FScriptExecutionNode> ExecNode = NodeEntry.Value.ExecNode;
 			double PureNodeEntryTime = 0.0;
+			double LastPureNodeTime = 0.0;
 			double NodeEntryTime = 0.0;
+			int32 TracePathIdx = 0;
 			for (auto EventIter = NodeEntry.Value.UnproccessedEvents.CreateIterator(); EventIter; ++EventIter)
 			{
 				switch(EventIter->GetType())
 				{
+					case EScriptInstrumentation::PushState:
+					{
+						LastPureNodeTime = EventIter->GetTime();
+						break;
+					}
 					case EScriptInstrumentation::PureNodeEntry:
 					{
 						if (PureNodeEntryTime == 0.0)
 						{
 							PureNodeEntryTime = EventIter->GetTime();
 						}
+
+						const int32 ScriptCodeOffset = EventIter->GetScriptCodeOffset();
+						if (const UEdGraphPin* Pin = FunctionContext->GetPinFromCodeLocation(ScriptCodeOffset))
+						{
+							if (FunctionContext->HasProfilerDataForNode(Pin->GetOwningNode()->GetFName()))
+							{
+								TSharedPtr<FScriptExecutionNode> PureNode = NodeEntry.Value.FunctionContext->GetProfilerDataForNode(Pin->GetOwningNode()->GetFName());
+								TSharedPtr<FScriptPerfData> PerfData = PureNode->GetPerfDataByInstanceAndTracePath(InstanceName, NodeEntry.Value.InputTracePaths.Pop());
+								PerfData->AddEventTiming(0.0f, EventIter->GetTime() - LastPureNodeTime);
+							}
+						}
+
+						LastPureNodeTime = EventIter->GetTime();
 						break;
 					}
 					case EScriptInstrumentation::NodeEntry:
@@ -619,7 +895,8 @@ bool FScriptEventPlayback::Process(const TArray<FScriptInstrumentedEvent>& Event
 					case EScriptInstrumentation::NodeExit:
 					{
 						check(NodeEntryTime != 0.0 );
-						PerfData = ExecNode->GetPerfDataByInstanceAndTracePath(InstanceName, NodeEntry.Value.EntryTracePaths.Pop());
+						const FTracePath NodeTracePath = NodeEntry.Value.EntryTracePaths[TracePathIdx++];
+						TSharedPtr<FScriptPerfData> PerfData = ExecNode->GetPerfDataByInstanceAndTracePath(InstanceName, NodeTracePath);
 						double PureNodeDuration = PureNodeEntryTime != 0.0 ? (NodeEntryTime - PureNodeEntryTime) : 0.0;
 						double NodeDuration = EventIter->GetTime() - NodeEntryTime;
 						PerfData->AddEventTiming(PureNodeDuration, NodeDuration);
@@ -629,9 +906,8 @@ bool FScriptEventPlayback::Process(const TArray<FScriptInstrumentedEvent>& Event
 				}
 			}
 		}
-		return true;
 	}
-	return false;
+	return ProcessingState != EEventProcessingResult::Failed;
 }
 
 #undef LOCTEXT_NAMESPACE
diff --git a/Engine/Source/Developer/BlueprintProfiler/Public/BlueprintProfiler.h b/Engine/Source/Developer/BlueprintProfiler/Public/BlueprintProfiler.h
--- a/Engine/Source/Developer/BlueprintProfiler/Public/BlueprintProfiler.h
+++ b/Engine/Source/Developer/BlueprintProfiler/Public/BlueprintProfiler.h
@@ -34,7 +34,7 @@ public:
 	virtual FOnBPStatGraphLayoutChanged& GetGraphLayoutChangedDelegate() { return GraphLayoutChangedDelegate; }
 	virtual TSharedPtr<FBlueprintExecutionContext> GetBlueprintContext(const FString& BlueprintClassPath) override;
 	virtual TSharedPtr<FScriptExecutionNode> GetProfilerDataForNode(const UEdGraphNode* GraphNode) override;
-	virtual void MapBlueprintInstance(TSharedPtr<FBlueprintExecutionContext> BlueprintContext, const FString& InstancePath) override;
+	virtual FName MapBlueprintInstance(TSharedPtr<FBlueprintExecutionContext> BlueprintContext, const FString& InstancePath) override;
 	virtual bool HasDataForInstance(const UObject* Instance) const override;
 	virtual void ProcessEventProfilingData() override;
 #endif
@@ -81,6 +81,8 @@ protected:
 #if WITH_EDITOR
 	/** PIE Active */
 	bool bPIEActive;
+	/** Suspended script events */
+	TMap<FName, TSharedPtr<FScriptEventPlayback>> SuspendedEvents;
 	/** Cached object path and code offset lookup to UObjects */
 	TMap<FString, TSharedPtr<FBlueprintExecutionContext>> PathToBlueprintContext;
 #endif
diff --git a/Engine/Source/Developer/BlueprintProfiler/Public/BlueprintProfilerModule.h b/Engine/Source/Developer/BlueprintProfiler/Public/BlueprintProfilerModule.h
--- a/Engine/Source/Developer/BlueprintProfiler/Public/BlueprintProfilerModule.h
+++ b/Engine/Source/Developer/BlueprintProfiler/Public/BlueprintProfilerModule.h
@@ -37,8 +37,8 @@ public:
 	/** Returns the current profiling event data for node */
 	virtual TSharedPtr<class FScriptExecutionNode> GetProfilerDataForNode(const UEdGraphNode* GraphNode) = 0;
 
-	/** Maps a blueprint instance and links to blueprint profiling data */
-	virtual void MapBlueprintInstance(TSharedPtr<FBlueprintExecutionContext> BlueprintContext, const FString& InstancePath) = 0;
+	/** Maps a blueprint instance and links to blueprint profiling data, returns the correct instance name */
+	virtual FName MapBlueprintInstance(TSharedPtr<FBlueprintExecutionContext> BlueprintContext, const FString& InstancePath) = 0;
 
 	/** Returns if the profiler currently retains any data for the specified instance */
 	virtual bool HasDataForInstance(const UObject* Instance) const { return false; }
diff --git a/Engine/Source/Developer/BlueprintProfiler/Public/ScriptInstrumentationCapture.h b/Engine/Source/Developer/BlueprintProfiler/Public/ScriptInstrumentationCapture.h
--- a/Engine/Source/Developer/BlueprintProfiler/Public/ScriptInstrumentationCapture.h
+++ b/Engine/Source/Developer/BlueprintProfiler/Public/ScriptInstrumentationCapture.h
@@ -59,7 +59,10 @@ public:
 	bool IsNewInstance() const { return EventType == EScriptInstrumentation::Instance; }
 
 	/** Returns if this event represents a change in active event */
-	bool IsEvent() const { return EventType == EScriptInstrumentation::Event || EventType == EScriptInstrumentation::Stop; }
+	bool IsEvent() const { return EventType == EScriptInstrumentation::Event||EventType == EScriptInstrumentation::ResumeEvent; }
+
+	/** Returns if this event represents a active event resuming execution */
+	bool IsResumeEvent() const { return EventType == EScriptInstrumentation::ResumeEvent; }
 
 	/** Returns if this event represents a node execution event */
 	bool IsNodeTiming() const {	return EventType == EScriptInstrumentation::NodeEntry || EventType == EScriptInstrumentation::NodeExit;	}
diff --git a/Engine/Source/Developer/BlueprintProfiler/Public/ScriptInstrumentationPlayback.h b/Engine/Source/Developer/BlueprintProfiler/Public/ScriptInstrumentationPlayback.h
--- a/Engine/Source/Developer/BlueprintProfiler/Public/ScriptInstrumentationPlayback.h
+++ b/Engine/Source/Developer/BlueprintProfiler/Public/ScriptInstrumentationPlayback.h
@@ -35,11 +35,11 @@ public:
 	/** Remaps PIE actor instance paths to editor actor instances */
 	BLUEPRINTPROFILER_API FName RemapInstancePath(const FName InstanceName) const ;
 
-	/** Returns the function context matching the supplied name providing it exists */
-	TSharedPtr<class FBlueprintFunctionContext> FindFunctionContext(const FName FunctionNameIn) const;
+	/** Returns the function name containing the event */
+	FName GetEventFunctionName(const FName EventName) const;
 
-	/** Returns the function context matching the supplied name, creates a new one if neccesary */
-	TSharedPtr<FBlueprintFunctionContext> GetOrCreateFunctionContext(const FName FunctionNameIn);
+	/** Returns the function context matching the supplied name providing it exists */
+	TSharedPtr<class FBlueprintFunctionContext> GetFunctionContext(const FName FunctionNameIn) const;
 
 	/** Returns if this context has an execution node for the specified graph node */
 	BLUEPRINTPROFILER_API bool HasProfilerDataForNode(const UEdGraphNode* GraphNode) const;
@@ -52,11 +52,17 @@ public:
 
 private:
 
+	/** Creates and returns a new function context, returns existing if present  */
+	TSharedPtr<FBlueprintFunctionContext> CreateFunctionContext(const FName FunctionNameIn);
+
 	/** Walks the blueprint and maps events and functions ready for profiling data */
 	bool MapBlueprintExecution();
 
+	/** Locates and generates event contexts for any input events passed in */
+	void CreateInputEvents(const TArray<UK2Node*>& InputEventNodes);
+
 	/** Maps each blueprint node following execution wires */
-	TSharedPtr<FScriptExecutionNode> MapNodeExecution(UEdGraphNode* NodeToMap);
+	TSharedPtr<FScriptExecutionNode> MapNodeExecution(TSharedPtr<FBlueprintFunctionContext> FunctionContext, UEdGraphNode* NodeToMap);
 
 private:
 
@@ -96,15 +102,31 @@ public:
 	/** Returns the graph node at the specified script offset */
 	const UEdGraphNode* GetNodeFromCodeLocation(const int32 ScriptOffset);
 
+	/** Returns the pin at the specified script offset */
+	const UEdGraphPin* GetPinFromCodeLocation(const int32 ScriptOffset);
+
 	/** Returns the script offset for the specified pin */
 	const int32 GetCodeLocationFromPin(const UEdGraphPin* Pin) const;
 
+	/** Returns all registered script offsets for the specified pin */
+	void GetAllCodeLocationsFromPin(const UEdGraphPin* Pin, TArray<int32>& OutCodeLocations) const;
+
+	/** Returns the pure node script code range for the specified node */
+	FInt32Range GetPureNodeScriptCodeRange(const UEdGraphNode* Node) const;
+
 	/** Returns true if this function context contains an execution node matching the node name */
 	bool HasProfilerDataForNode(const FName NodeName) const;
 
 	/** Returns the execution node representing the node name */
 	TSharedPtr<FScriptExecutionNode> GetProfilerDataForNode(const FName NodeName);
 
+protected:
+
+	friend FBlueprintExecutionContext;
+
+	/** Utility to create execution node */
+	TSharedPtr<FScriptExecutionNode> CreateExecutionNode(FScriptExecNodeParams& InitParams);
+
 private:
 
 	/** The UFunction this context represents */
@@ -113,6 +135,8 @@ private:
 	TWeakObjectPtr<UBlueprintGeneratedClass> BlueprintClass;
 	/** ScriptCode offset to node cache */
 	TMap<int32, TWeakObjectPtr<const UEdGraphNode>> ScriptOffsetToNodes;
+	/** ScriptCode offset to pin cache */
+	TMap<int32, TWeakObjectPtr<const UEdGraphPin>> ScriptOffsetToPins;
 	/** Execution nodes containing profiling data */
 	TMap<FName, TSharedPtr<FScriptExecutionNode>> ExecutionNodes;
 };
@@ -120,15 +144,34 @@ private:
 //////////////////////////////////////////////////////////////////////////
 // FScriptEventPlayback
 
-class FScriptEventPlayback
+namespace EEventProcessingResult
+{
+	enum Type
+	{
+		None = 0,
+		Failed,
+		Success,
+		Suspended
+	};
+};
+
+class FScriptEventPlayback : public TSharedFromThis<FScriptEventPlayback>
 {
 public:
 
-	FScriptEventPlayback(TSharedPtr<FBlueprintExecutionContext> BlueprintContextIn)
-		: BlueprintContext(BlueprintContextIn)
+	FScriptEventPlayback(TSharedPtr<FBlueprintExecutionContext> BlueprintContextIn, const FName InstanceNameIn)
+		: InstanceName(InstanceNameIn)
+		, ProcessingState(EEventProcessingResult::None)
+		, BlueprintContext(BlueprintContextIn)
 	{
 	}
 
+	/** Returns the current processing state */
+	EEventProcessingResult::Type GetProcessingState() const { return ProcessingState; }
+
+	/** Returns if current processing state is suspended */
+	bool IsSuspended() const { return ProcessingState == EEventProcessingResult::Suspended; }
+
 	/** Processes the event and cleans up any errant signals */
 	bool Process(const TArray<class FScriptInstrumentedEvent>& Events, const int32 Start, const int32 Stop);
 
@@ -141,6 +184,14 @@ private:
 	FName InstanceName;
 	/** Event name */
 	FName EventName;
+	/** Processing state */
+	EEventProcessingResult::Type ProcessingState;
+	/** Current tracepath state */
+	FTracePath TracePath;
+	/** Current tracepath stack */
+	TArray<FTracePath> TraceStack;
+	/** Event Timings */
+	TMultiMap<FName, double> EventTimings;
 	/** Blueprint exec context */
 	TSharedPtr<FBlueprintExecutionContext> BlueprintContext;
 
diff --git a/Engine/Source/Editor/BlueprintGraph/Classes/K2Node_Event.h b/Engine/Source/Editor/BlueprintGraph/Classes/K2Node_Event.h
--- a/Engine/Source/Editor/BlueprintGraph/Classes/K2Node_Event.h
+++ b/Engine/Source/Editor/BlueprintGraph/Classes/K2Node_Event.h
@@ -61,6 +61,7 @@ class UK2Node_Event : public UK2Node_EditablePinBase
 	BLUEPRINTGRAPH_API virtual FString GetDeprecationMessage() const override;
 	BLUEPRINTGRAPH_API virtual UObject* GetJumpTargetForDoubleClick() const override;
 	BLUEPRINTGRAPH_API virtual FName GetPaletteIcon(FLinearColor& OutColor) const override{ return TEXT("GraphEditor.Event_16x"); }
+	BLUEPRINTGRAPH_API virtual FString GetFindReferenceSearchString() const override;
 	//~ End UEdGraphNode Interface
 
 	//~ Begin UK2Node Interface
diff --git a/Engine/Source/Editor/BlueprintGraph/Private/CallFunctionHandler.cpp b/Engine/Source/Editor/BlueprintGraph/Private/CallFunctionHandler.cpp
--- a/Engine/Source/Editor/BlueprintGraph/Private/CallFunctionHandler.cpp
+++ b/Engine/Source/Editor/BlueprintGraph/Private/CallFunctionHandler.cpp
@@ -383,6 +383,13 @@ void FKCHandler_CallFunction::CreateFunctionCallStatement(FKismetFunctionContext
 				pSrcEventNode = CompilerContext.CallsIntoUbergraph.Find(Node);
 			}
 
+			const bool bInstrumentFunctionEntry = Context.IsInstrumentationRequired() && (bIsLatent || IsUserFunctionCall(Node));
+			if (bInstrumentFunctionEntry)
+			{
+				FBlueprintCompiledStatement& PushState = Context.AppendStatementForNode(Node);
+				PushState.Type = pSrcEventNode ? KCST_InstrumentedStateReset : KCST_InstrumentedStatePush;
+			}
+
 			// Iterate over all the contexts this functions needs to be called on, and emit a call function statement for each
 			for (auto TargetListIt = ContextTerms.CreateIterator(); TargetListIt; ++TargetListIt)
 			{
@@ -427,6 +434,11 @@ void FKCHandler_CallFunction::CreateFunctionCallStatement(FKismetFunctionContext
 			// Create the exit from this node if there is one
 			if (bIsLatent)
 			{
+				if (bInstrumentFunctionEntry)
+				{
+					FBlueprintCompiledStatement& SuspendState = Context.AppendStatementForNode(Node);
+					SuspendState.Type = KCST_InstrumentedStateSuspend;
+				}
 				// End this thread of execution; the latent function will resume it at some point in the future
 				FBlueprintCompiledStatement& PopStatement = Context.AppendStatementForNode(Node);
 				PopStatement.Type = KCST_EndOfThread;
@@ -436,6 +448,11 @@ void FKCHandler_CallFunction::CreateFunctionCallStatement(FKismetFunctionContext
 				// Generate the output impulse from this node
 				if (!IsCalledFunctionPure(Node))
 				{
+					if (bInstrumentFunctionEntry)
+					{
+						FBlueprintCompiledStatement& PopState = Context.AppendStatementForNode(Node);
+						PopState.Type = KCST_InstrumentedStatePop;
+					}
 					GenerateSimpleThenGoto(Context, *Node);
 				}
 			}
@@ -739,6 +756,26 @@ FString FKCHandler_CallFunction::GetFunctionNameFromNode(UEdGraphNode* Node) con
 	}
 }
 
+bool FKCHandler_CallFunction::IsUserFunctionCall(UEdGraphNode* Node) const
+{
+	bool bResult = false;
+	if (UK2Node_CallFunction* CallFunctionNode = Cast<UK2Node_CallFunction>(Node))
+	{
+		TArray<UEdGraph*> Graphs;
+		CompilerContext.Blueprint->GetAllGraphs(Graphs);
+		const FName CalledGraph = CallFunctionNode->FunctionReference.GetMemberName();
+		for (auto Graph : Graphs)
+		{
+			if (CalledGraph == Graph->GetFName())
+			{
+				bResult = true;
+				break;
+			}
+		}
+	}
+	return bResult;
+}
+
 #if _MSC_VER
 #pragma warning(pop)
 #endif
diff --git a/Engine/Source/Editor/BlueprintGraph/Private/EdGraphSchema_K2.cpp b/Engine/Source/Editor/BlueprintGraph/Private/EdGraphSchema_K2.cpp
--- a/Engine/Source/Editor/BlueprintGraph/Private/EdGraphSchema_K2.cpp
+++ b/Engine/Source/Editor/BlueprintGraph/Private/EdGraphSchema_K2.cpp
@@ -4712,7 +4712,14 @@ void UEdGraphSchema_K2::GetGraphDisplayInformation(const UEdGraph& Graph, /*out*
 
 	if( GEditor && GetDefault<UEditorStyleSettings>()->bShowFriendlyNames )
 	{
-		DisplayInfo.DisplayName = FText::FromString(FName::NameToDisplayString( DisplayInfo.PlainName.ToString(), false ));
+		if (GraphType == GT_Function && Function)
+		{
+			DisplayInfo.DisplayName = GetFriendlySignatureName(Function);
+		}
+		else
+		{
+			DisplayInfo.DisplayName = FText::FromString(FName::NameToDisplayString(DisplayInfo.PlainName.ToString(), false));
+		}
 	}
 	else
 	{
diff --git a/Engine/Source/Editor/BlueprintGraph/Private/K2Node_CallFunction.cpp b/Engine/Source/Editor/BlueprintGraph/Private/K2Node_CallFunction.cpp
--- a/Engine/Source/Editor/BlueprintGraph/Private/K2Node_CallFunction.cpp
+++ b/Engine/Source/Editor/BlueprintGraph/Private/K2Node_CallFunction.cpp
@@ -1635,9 +1635,10 @@ void UK2Node_CallFunction::FixupSelfMemberContext()
 		}
 		return bIsChildOf;
 	};
-	
 
 	UClass* MemberClass = FunctionReference.GetMemberParentClass();
+	const bool bIsLocalMacro = Blueprint && Blueprint->MacroGraphs.Contains(GetGraph());
+	ensureMsgf(FunctionReference.IsSelfContext() || (MemberClass != nullptr) || bIsLocalMacro, TEXT("Unknown member class in %s"), *GetPathName());
 	if (FunctionReference.IsSelfContext())
 	{
 		if (MemberClass == nullptr)
@@ -1664,7 +1665,7 @@ void UK2Node_CallFunction::FixupSelfMemberContext()
 		// identically named function)... if there is no function matching this 
 		// reference, then the node will produce an error later during compilation
 	}
-	else if (ensure(MemberClass != nullptr))
+	else if (MemberClass != nullptr)
 	{
 		if (IsBlueprintOfType(MemberClass))
 		{
diff --git a/Engine/Source/Editor/BlueprintGraph/Private/K2Node_Event.cpp b/Engine/Source/Editor/BlueprintGraph/Private/K2Node_Event.cpp
--- a/Engine/Source/Editor/BlueprintGraph/Private/K2Node_Event.cpp
+++ b/Engine/Source/Editor/BlueprintGraph/Private/K2Node_Event.cpp
@@ -795,6 +795,18 @@ UObject* UK2Node_Event::GetJumpTargetForDoubleClick() const
 	return NULL;
 }
 
+FString UK2Node_Event::GetFindReferenceSearchString() const
+{
+	FString FunctionName = EventReference.GetMemberName().ToString(); // If we fail to find the function, still want to search for its expected name.
+
+	if (UFunction* Function = EventReference.ResolveMember<UFunction>(GetBlueprintClassFromNode()))
+	{
+		FunctionName = UEdGraphSchema_K2::GetFriendlySignatureName(Function).ToString();
+	}
+
+	return FunctionName;
+}
+
 bool UK2Node_Event::AreEventNodesIdentical(const UK2Node_Event* InNodeA, const UK2Node_Event* InNodeB)
 {
 	return InNodeA->EventReference.GetMemberName() == InNodeB->EventReference.GetMemberName()
diff --git a/Engine/Source/Editor/BlueprintGraph/Private/K2Node_TemporaryVariable.cpp b/Engine/Source/Editor/BlueprintGraph/Private/K2Node_TemporaryVariable.cpp
--- a/Engine/Source/Editor/BlueprintGraph/Private/K2Node_TemporaryVariable.cpp
+++ b/Engine/Source/Editor/BlueprintGraph/Private/K2Node_TemporaryVariable.cpp
@@ -179,6 +179,9 @@ void UK2Node_TemporaryVariable::GetMenuActions(FBlueprintActionDatabaseRegistrar
 	ActionRegistrar.AddBlueprintAction(ActionKey, MakeTempVarNodeSpawner(FEdGraphPinType(K2Schema->PC_String, TEXT(""), nullptr, /*bIsArray =*/ true, /*bIsReference =*/false), /*bIsPersistent =*/false));
 	ActionRegistrar.AddBlueprintAction(ActionKey, MakeTempVarNodeSpawner(FEdGraphPinType(K2Schema->PC_Text, TEXT(""), nullptr, /*bIsArray =*/false, /*bIsReference =*/false), /*bIsPersistent =*/false));
 	ActionRegistrar.AddBlueprintAction(ActionKey, MakeTempVarNodeSpawner(FEdGraphPinType(K2Schema->PC_Text, TEXT(""), nullptr, /*bIsArray =*/ true, /*bIsReference =*/false), /*bIsPersistent =*/false));
+	
+	ActionRegistrar.AddBlueprintAction(ActionKey, MakeTempVarNodeSpawner(FEdGraphPinType(K2Schema->PC_Wildcard, TEXT(""), nullptr, /*bIsArray =*/false, /*bIsReference =*/false), /*bIsPersistent =*/false));
+	ActionRegistrar.AddBlueprintAction(ActionKey, MakeTempVarNodeSpawner(FEdGraphPinType(K2Schema->PC_Wildcard, TEXT(""), nullptr, /*bIsArray =*/ true, /*bIsReference =*/false), /*bIsPersistent =*/false));
 
 	UScriptStruct* VectorStruct  = TBaseStructure<FVector>::Get();
 	ActionRegistrar.AddBlueprintAction(ActionKey, MakeTempVarNodeSpawner(FEdGraphPinType(K2Schema->PC_Struct, TEXT("Vector"), VectorStruct, /*bIsArray =*/false, /*bIsReference =*/false), /*bIsPersistent =*/false));
diff --git a/Engine/Source/Editor/BlueprintGraph/Public/CallFunctionHandler.h b/Engine/Source/Editor/BlueprintGraph/Public/CallFunctionHandler.h
--- a/Engine/Source/Editor/BlueprintGraph/Public/CallFunctionHandler.h
+++ b/Engine/Source/Editor/BlueprintGraph/Public/CallFunctionHandler.h
@@ -53,6 +53,7 @@ private:
 	virtual FString GetFunctionNameFromNode(UEdGraphNode* Node) const;
 	UClass* GetCallingContext(FKismetFunctionContext& Context, UEdGraphNode* Node);
 	UClass* GetTrueCallingClass(FKismetFunctionContext& Context, UEdGraphPin* SelfPin);
+	bool IsUserFunctionCall(UEdGraphNode* Node) const;
 
 public:
 
diff --git a/Engine/Source/Editor/GraphEditor/Private/BlueprintConnectionDrawingPolicy.cpp b/Engine/Source/Editor/GraphEditor/Private/BlueprintConnectionDrawingPolicy.cpp
--- a/Engine/Source/Editor/GraphEditor/Private/BlueprintConnectionDrawingPolicy.cpp
+++ b/Engine/Source/Editor/GraphEditor/Private/BlueprintConnectionDrawingPolicy.cpp
@@ -102,7 +102,7 @@ void FKismetConnectionDrawingPolicy::BuildExecutionRoadmap()
 			{
 				if (TestObject == ActiveObject)
 				{
-					UEdGraphPin* AssociatedPin = DebugData.FindExecPinFromCodeLocation(Sample.Function.Get(), Sample.Offset);
+					UEdGraphPin* AssociatedPin = DebugData.FindSourcePinFromCodeLocation(Sample.Function.Get(), Sample.Offset);
 
 					if (UEdGraphNode* Node = DebugData.FindSourceNodeFromCodeLocation(Sample.Function.Get(), Sample.Offset, /*bAllowImpreciseHit=*/ false))
 					{
diff --git a/Engine/Source/Editor/Kismet/Private/BlueprintEditor.cpp b/Engine/Source/Editor/Kismet/Private/BlueprintEditor.cpp
--- a/Engine/Source/Editor/Kismet/Private/BlueprintEditor.cpp
+++ b/Engine/Source/Editor/Kismet/Private/BlueprintEditor.cpp
@@ -8039,6 +8039,8 @@ FReply FBlueprintEditor::OnSpawnGraphNodeByShortcut(FInputChord InChord, const F
 		return FReply::Handled();
 	}
 
+	FScopedTransaction Transaction(LOCTEXT("AddNode", "Add Node"));
+
 	TArray<UEdGraphNode*> OutNodes;
 	FVector2D NodeSpawnPos = InPosition;
 	FBlueprintSpawnNodeCommands::Get().GetGraphActionByChord(InChord, InGraph, NodeSpawnPos, OutNodes);
@@ -8055,6 +8057,10 @@ FReply FBlueprintEditor::OnSpawnGraphNodeByShortcut(FInputChord InChord, const F
 	{
 		Graph->SelectNodeSet(NodesToSelect, /*bFromUI =*/true);
 	}
+	else
+	{
+		Transaction.Cancel();
+	}
 
 	return FReply::Handled();
 }
diff --git a/Engine/Source/Editor/Kismet/Private/BlueprintEditorCommands.cpp b/Engine/Source/Editor/Kismet/Private/BlueprintEditorCommands.cpp
--- a/Engine/Source/Editor/Kismet/Private/BlueprintEditorCommands.cpp
+++ b/Engine/Source/Editor/Kismet/Private/BlueprintEditorCommands.cpp
@@ -250,6 +250,7 @@ public:
 		UK2Node_MacroInstance* MacroInstanceNode = Cast<UK2Node_MacroInstance>(UBlueprintNodeSpawner::Create(UK2Node_MacroInstance::StaticClass())->Invoke(InDestGraph, Bindings, InOutDestPosition));
 		MacroInstanceNode->SetMacroGraph(MacroGraph);
 		MacroInstanceNode->ReconstructNode();
+		OutNodes.Add(MacroInstanceNode);
 	}
 	// End of FNodeSpawnInfo interface
 
diff --git a/Engine/Source/Editor/Kismet/Private/Profiler/BPProfilerStatisticWidgets.cpp b/Engine/Source/Editor/Kismet/Private/Profiler/BPProfilerStatisticWidgets.cpp
--- a/Engine/Source/Editor/Kismet/Private/Profiler/BPProfilerStatisticWidgets.cpp
+++ b/Engine/Source/Editor/Kismet/Private/Profiler/BPProfilerStatisticWidgets.cpp
@@ -128,6 +128,15 @@ TSharedRef<SWidget> FBPProfilerStatWidget::GenerateColumnWidget(FName ColumnName
 					.Style(FEditorStyle::Get(), "HoverOnlyHyperlink")
 					.ToolTipText(ExecNode->GetToolTipText())
 					.OnNavigate(this, &FBPProfilerStatWidget::NavigateTo)
+		#if TRACEPATH_DEBUG
+				]
+				+SHorizontalBox::Slot()
+				.AutoWidth()
+				.Padding(FMargin(5,0))
+				[
+					SNew(STextBlock)
+					.Text(FText::FromString(*WidgetTracePath.GetPathString()))
+		#endif
 				];
 		}
 		else
@@ -233,6 +242,27 @@ void FBPProfilerStatWidget::GenerateExecNodeWidgets(const FName InstanceName, co
 		}
 		else
 		{
+			// @TODO - filter out pure nodes?
+			if (!ExecNode->IsPureNode())
+			{
+				// Get the full pure node chain associated with this exec node.
+				TMap<int32, TSharedPtr<FScriptExecutionNode>> AllPureNodes;
+				ExecNode->GetAllPureNodes(AllPureNodes);
+
+				// Sort pure nodes by script offset (execution order).
+				AllPureNodes.KeySort(TLess<int32>());
+
+				// Build trace path, tree view node widget and register perf stats for tracking.
+				FTracePath PureTracePath(WidgetTracePath);
+				for (auto Iter : AllPureNodes)
+				{
+					PureTracePath.AddExitPin(Iter.Key);
+					TSharedPtr<FBPProfilerStatWidget> NewPureChildNode = MakeShareable<FBPProfilerStatWidget>(new FBPProfilerStatWidget(Iter.Value, PureTracePath));
+					NewPureChildNode->GenerateExecNodeWidgets(InstanceName, FilterGraph);
+					CachedChildren.Add(NewPureChildNode);
+				}
+			}
+
 			for (auto Iter : ExecNode->GetChildNodes())
 			{
 				// Filter out events based on graph
diff --git a/Engine/Source/Editor/Kismet/Private/Profiler/EventExecution.cpp b/Engine/Source/Editor/Kismet/Private/Profiler/EventExecution.cpp
--- a/Engine/Source/Editor/Kismet/Private/Profiler/EventExecution.cpp
+++ b/Engine/Source/Editor/Kismet/Private/Profiler/EventExecution.cpp
@@ -7,6 +7,12 @@
 //////////////////////////////////////////////////////////////////////////
 // FScriptNodeExecLinkage
 
+void FScriptNodeExecLinkage::AddPureNode(const int32 PinScriptOffset, TSharedPtr<class FScriptExecutionNode> PureNode)
+{
+	TSharedPtr<class FScriptExecutionNode>& Result = PureNodes.FindOrAdd(PinScriptOffset);
+	Result = PureNode;
+}
+
 void FScriptNodeExecLinkage::AddLinkedNode(const int32 PinScriptOffset, TSharedPtr<class FScriptExecutionNode> LinkedNode)
 {
 	TSharedPtr<class FScriptExecutionNode>& Result = LinkedNodes.FindOrAdd(PinScriptOffset);
@@ -16,6 +22,17 @@ void FScriptNodeExecLinkage::AddLinkedNode(const int32 PinScriptOffset, TSharedP
 //////////////////////////////////////////////////////////////////////////
 // FScriptNodePerfData
 
+bool FScriptNodePerfData::HasPerfDataForInstanceAndTracePath(FName InstanceName, const FTracePath& TracePath) const
+{
+	const TMap<const uint32, TSharedPtr<FScriptPerfData>>* InstanceMapPtr = InstanceInputPinToPerfDataMap.Find(InstanceName);
+	if (InstanceMapPtr)
+	{
+		return InstanceMapPtr->Contains(TracePath);
+	}
+
+	return false;
+}
+
 TSharedPtr<FScriptPerfData> FScriptNodePerfData::GetPerfDataByInstanceAndTracePath(FName InstanceName, const FTracePath& TracePath)
 { 
 	TMap<const uint32, TSharedPtr<FScriptPerfData>>& InstanceMap = InstanceInputPinToPerfDataMap.FindOrAdd(InstanceName);
@@ -62,6 +79,19 @@ FScriptExecutionNode::FScriptExecutionNode()
 {
 }
 
+FScriptExecutionNode::FScriptExecutionNode(const FScriptExecNodeParams& InitParams)
+	: NodeName(InitParams.NodeName)
+	, OwningGraphName(InitParams.OwningGraphName)
+	, NodeFlags(InitParams.NodeFlags)
+	, ObservedObject(InitParams.ObservedObject)
+	, DisplayName(InitParams.DisplayName)
+	, Tooltip(InitParams.Tooltip)
+	, IconColor(InitParams.IconColor)
+	, Icon(InitParams.Icon)
+	, bExpansionState(false)
+{
+}
+
 bool FScriptExecutionNode::operator == (const FScriptExecutionNode& NodeIn) const
 {
 	return NodeName == NodeIn.NodeName;
@@ -168,7 +198,7 @@ void FScriptExecutionNode::RefreshStats_Internal(const FTracePath& TracePath, FE
 						TSharedPtr<FScriptPerfData> InstancePerfData = GetPerfDataByInstanceAndTracePath(InstanceIter, TracePath);
 						InstancePerfData->Reset();
 						BlueprintPooledStats.Add(InstancePerfData);
-				
+
 						for (auto ChildIter : ChildNodes)
 						{
 							TSharedPtr<FScriptPerfData> ChildPerfData = ChildIter->GetPerfDataByInstanceAndTracePath(InstanceIter, TracePath);
@@ -200,6 +230,21 @@ void FScriptExecutionNode::RefreshStats_Internal(const FTracePath& TracePath, FE
 		// Refresh All links
 		if (bRefreshLinkStats)
 		{
+			if (!IsPureNode())
+			{
+				TMap<int32, TSharedPtr<FScriptExecutionNode>> AllPureNodes;
+				GetAllPureNodes(AllPureNodes);
+
+				AllPureNodes.KeySort(TLess<int32>());
+
+				FTracePath PureTracePath(TracePath);
+				for (auto PureIter : AllPureNodes)
+				{
+					PureTracePath.AddExitPin(PureIter.Key);
+					PureIter.Value->RefreshStats(PureTracePath);
+				}
+			}
+
 			for (auto LinkIter : LinkedNodes)
 			{
 				FTracePath LinkTracePath(TracePath);
@@ -235,6 +280,24 @@ void FScriptExecutionNode::GetAllExecNodes(TMap<FName, TSharedPtr<FScriptExecuti
 	}
 }
 
+void FScriptExecutionNode::GetAllPureNodes(TMap<int32, TSharedPtr<class FScriptExecutionNode>>& PureNodesOut)
+{
+	GetAllPureNodes_Internal(PureNodesOut, PureNodeScriptCodeRange);
+}
+
+void FScriptExecutionNode::GetAllPureNodes_Internal(TMap<int32, TSharedPtr<class FScriptExecutionNode>>& PureNodesOut, const FInt32Range& ScriptCodeRange)
+{
+	for (auto PureIter : PureNodes)
+	{
+		PureIter.Value->GetAllPureNodes_Internal(PureNodesOut, ScriptCodeRange);
+
+		if (ScriptCodeRange.Contains(PureIter.Key))
+		{
+			PureNodesOut.Add(PureIter.Key, PureIter.Value);
+		}
+	}
+}
+
 void FScriptExecutionNode::NavigateToObject() const
 {
 	if (ObservedObject.IsValid())
diff --git a/Engine/Source/Editor/Kismet/Private/Profiler/SBlueprintProfilerView.cpp b/Engine/Source/Editor/Kismet/Private/Profiler/SBlueprintProfilerView.cpp
--- a/Engine/Source/Editor/Kismet/Private/Profiler/SBlueprintProfilerView.cpp
+++ b/Engine/Source/Editor/Kismet/Private/Profiler/SBlueprintProfilerView.cpp
@@ -50,8 +50,15 @@ void SBlueprintProfilerView::Construct(const FArguments& InArgs)
 
 void SBlueprintProfilerView::OnToggleProfiler(bool bEnabled)
 {
-	StatusText = bEnabled ? LOCTEXT("ProfilerNoDataText", "The Blueprint Profiler is active but does not currently have any data to display") :
-							LOCTEXT("ProfilerInactiveText", "The Blueprint Profiler is currently Inactive");
+	if (GetDefault<UEditorExperimentalSettings>()->bBlueprintPerformanceAnalysisTools)
+	{
+		StatusText = bEnabled ? LOCTEXT("ProfilerNoDataText", "The Blueprint Profiler is active but does not currently have any data to display") :
+								LOCTEXT("ProfilerInactiveText", "The Blueprint Profiler is currently Inactive");
+	}
+	else
+	{
+		StatusText = LOCTEXT("DisabledProfilerText", "The Blueprint Profiler is experimental and is currently not enabled in the editor preferences");
+	}
 	UpdateActiveProfilerWidget();
 }
 
@@ -174,10 +181,13 @@ TSharedRef<SWidget> SBlueprintProfilerView::CreateActiveStatisticWidget()
 {
 	// Get profiler status
 	EBlueprintPerfViewType::Type NewViewType = EBlueprintPerfViewType::None;
-	IBlueprintProfilerInterface* ProfilerInterface = FModuleManager::GetModulePtr<IBlueprintProfilerInterface>("BlueprintProfiler");
-	if (ProfilerInterface && ProfilerInterface->IsProfilerEnabled())
+	if (GetDefault<UEditorExperimentalSettings>()->bBlueprintPerformanceAnalysisTools)
 	{
-		NewViewType = CurrentViewType;
+		IBlueprintProfilerInterface* ProfilerInterface = FModuleManager::GetModulePtr<IBlueprintProfilerInterface>("BlueprintProfiler");
+		if (ProfilerInterface && ProfilerInterface->IsProfilerEnabled())
+		{
+			NewViewType = CurrentViewType;
+		}
 	}
 	switch(NewViewType)
 	{
diff --git a/Engine/Source/Editor/Kismet/Private/SGraphTitleBar.cpp b/Engine/Source/Editor/Kismet/Private/SGraphTitleBar.cpp
--- a/Engine/Source/Editor/Kismet/Private/SGraphTitleBar.cpp
+++ b/Engine/Source/Editor/Kismet/Private/SGraphTitleBar.cpp
@@ -34,7 +34,7 @@ FText SGraphTitleBar::GetTitleForOneCrumb(const UEdGraph* Graph)
 	Schema->GetGraphDisplayInformation(*Graph, /*out*/ DisplayInfo);
 
 	FFormatNamedArguments Args;
-	Args.Add(TEXT("BreadcrumbDisplayName"), DisplayInfo.PlainName);
+	Args.Add(TEXT("BreadcrumbDisplayName"), DisplayInfo.DisplayName);
 	Args.Add(TEXT("BreadcrumbNotes"), FText::FromString(DisplayInfo.GetNotesAsString()));
 	return FText::Format(LOCTEXT("BreadcrumbTitle", "{BreadcrumbDisplayName} {BreadcrumbNotes}"), Args);
 }
diff --git a/Engine/Source/Editor/Kismet/Private/SMyBlueprint.cpp b/Engine/Source/Editor/Kismet/Private/SMyBlueprint.cpp
--- a/Engine/Source/Editor/Kismet/Private/SMyBlueprint.cpp
+++ b/Engine/Source/Editor/Kismet/Private/SMyBlueprint.cpp
@@ -1201,7 +1201,7 @@ void SMyBlueprint::CollectAllActions(FGraphActionListBuilderBase& OutAllActions)
 
 		//@TODO: Should be a bit more generic (or the AnimGraph shouldn't be stored as a FunctionGraph...)
 		int32 SectionID = Graph->IsA<UAnimationGraph>() ? NodeSectionID::GRAPH : NodeSectionID::FUNCTION;
-		TSharedPtr<FEdGraphSchemaAction_K2Graph> NewFuncAction = MakeShareable(new FEdGraphSchemaAction_K2Graph(EEdGraphSchemaAction_K2Graph::Function, FunctionCategory, DisplayInfo.PlainName, DisplayInfo.Tooltip, bIsConstructionScript ? 2 : 1, SectionID));
+		TSharedPtr<FEdGraphSchemaAction_K2Graph> NewFuncAction = MakeShareable(new FEdGraphSchemaAction_K2Graph(EEdGraphSchemaAction_K2Graph::Function, FunctionCategory, DisplayInfo.DisplayName, DisplayInfo.Tooltip, bIsConstructionScript ? 2 : 1, SectionID));
 		NewFuncAction->FuncName = Graph->GetFName();
 		NewFuncAction->EdGraph = Graph;
 
@@ -1225,7 +1225,7 @@ void SMyBlueprint::CollectAllActions(FGraphActionListBuilderBase& OutAllActions)
 
 		FText MacroCategory = GetGraphCategory(Graph);
 
-		TSharedPtr<FEdGraphSchemaAction_K2Graph> NewMacroAction = MakeShareable(new FEdGraphSchemaAction_K2Graph(EEdGraphSchemaAction_K2Graph::Macro, MacroCategory, DisplayInfo.PlainName, DisplayInfo.Tooltip, 1, NodeSectionID::MACRO));
+		TSharedPtr<FEdGraphSchemaAction_K2Graph> NewMacroAction = MakeShareable(new FEdGraphSchemaAction_K2Graph(EEdGraphSchemaAction_K2Graph::Macro, MacroCategory, DisplayInfo.DisplayName, DisplayInfo.Tooltip, 1, NodeSectionID::MACRO));
 		NewMacroAction->FuncName = MacroName;
 		NewMacroAction->EdGraph = Graph;
 		OutAllActions.AddAction(NewMacroAction);
@@ -1343,7 +1343,7 @@ void SMyBlueprint::CollectAllActions(FGraphActionListBuilderBase& OutAllActions)
 		FGraphDisplayInfo DisplayInfo;
 		Graph->GetSchema()->GetGraphDisplayInformation(*Graph, DisplayInfo);
 
-		TSharedPtr<FEdGraphSchemaAction_K2Graph> NeUbergraphAction = MakeShareable(new FEdGraphSchemaAction_K2Graph(EEdGraphSchemaAction_K2Graph::Graph, FText::GetEmpty(), DisplayInfo.PlainName, DisplayInfo.Tooltip, 2, NodeSectionID::GRAPH));
+		TSharedPtr<FEdGraphSchemaAction_K2Graph> NeUbergraphAction = MakeShareable(new FEdGraphSchemaAction_K2Graph(EEdGraphSchemaAction_K2Graph::Graph, FText::GetEmpty(), DisplayInfo.DisplayName, DisplayInfo.Tooltip, 2, NodeSectionID::GRAPH));
 		NeUbergraphAction->FuncName = Graph->GetFName();
 		NeUbergraphAction->EdGraph = Graph;
 		OutAllActions.AddAction(NeUbergraphAction);
@@ -2188,7 +2188,7 @@ void SMyBlueprint::OnFindReference()
 	}
 	else if (FEdGraphSchemaAction_K2Event* EventAction = SelectionAsEvent())
 	{
-		SearchTerm = EventAction->GetMenuDescription().ToString();
+		SearchTerm = EventAction->NodeTemplate->GetFindReferenceSearchString();
 	}
 	else if (FEdGraphSchemaAction_K2InputAction* InputAction = SelectionAsInputAction())
 	{
diff --git a/Engine/Source/Editor/Kismet/Public/Profiler/EventExecution.h b/Engine/Source/Editor/Kismet/Public/Profiler/EventExecution.h
--- a/Engine/Source/Editor/Kismet/Public/Profiler/EventExecution.h
+++ b/Engine/Source/Editor/Kismet/Public/Profiler/EventExecution.h
@@ -13,16 +13,18 @@ namespace EScriptExecutionNodeFlags
 		Class				= 0x00000001,	// Class
 		Instance			= 0x00000002,	// Instance
 		Event				= 0x00000004,	// Event
+		CustomEvent			= 0x00000008,	// Custom Event
 		Container			= 0x00000007,	// Container type - Class, Instance or Event.
-		FunctionCall		= 0x00000008,	// Function Call
-		MacroCall			= 0x00000010,	// Macro Call
-		CallSite			= 0x00000018,	// Function / Macro call site
+		FunctionCall		= 0x00000010,	// Function Call
+		MacroCall			= 0x00000020,	// Macro Call
+		CallSite			= 0x00000030,	// Function / Macro call site
 		MacroNode			= 0x00000100,	// Macro Node
 		ConditionalBranch	= 0x00000200,	// Node has multiple exit pins using a jump
 		SequentialBranch	= 0x00000400,	// Node has multiple exit pins ran in sequence
 		BranchNode			= 0x00000600,	// BranchNode
 		Node				= 0x00000800,	// Node timing
-		ExecPin				= 0x00001000	// Exec pin dummy node
+		ExecPin				= 0x00001000,	// Exec pin dummy node
+		PureNode			= 0x00002000	// Pure node (no exec pins)
 	};
 }
 
@@ -63,15 +65,21 @@ public:
 	/** Sets the parent node */
 	void SetParentNode(TSharedPtr<class FScriptExecutionNode> InParentNode) { ParentNode = InParentNode; }
 
+	/** Returns the pure nodes map */
+	TMap<int32, TSharedPtr<class FScriptExecutionNode>>& GetPureNodes() { return PureNodes; }
+
+	/** Returns the number of pure nodes */
+	int32 GetNumPureNodes() const { return PureNodes.Num(); }
+
+	/** Add pure node */
+	void AddPureNode(const int32 PinScriptOffset, TSharedPtr<class FScriptExecutionNode> PureNode);
+
 	/** Returns the linked nodes map */
 	TMap<int32, TSharedPtr<class FScriptExecutionNode>>& GetLinkedNodes() { return LinkedNodes; }
 
 	/** Returns the number of linked nodes */
 	int32 GetNumLinkedNodes() const { return LinkedNodes.Num(); }
 
-	/** Returns the linked node for the specified pin script offset */
-	TSharedPtr<class FScriptExecutionNode> GetLinkedNodeByPinScriptOffset(const int32 PinScriptOffset) { return LinkedNodes[PinScriptOffset]; }
-
 	/** Add linked node */
 	void AddLinkedNode(const int32 PinScriptOffset, TSharedPtr<class FScriptExecutionNode> LinkedNode);
 
@@ -91,11 +99,12 @@ protected:
 
 	/** Parent node */
 	TSharedPtr<class FScriptExecutionNode> ParentNode;
+	/** Pure nodes */
+	TMap<int32, TSharedPtr<class FScriptExecutionNode>> PureNodes;
 	/** Linked nodes */
 	TMap<int32, TSharedPtr<class FScriptExecutionNode>> LinkedNodes;
 	/** Child nodes */
 	TArray<TSharedPtr<class FScriptExecutionNode>> ChildNodes;
-
 };
 
 //////////////////////////////////////////////////////////////////////////
@@ -105,7 +114,10 @@ class KISMET_API FScriptNodePerfData
 {
 public:
 
-	/** Get perf data for instance and tracepath */
+	/** Test whether or not perf data is available for the given instance/trace path */
+	bool HasPerfDataForInstanceAndTracePath(FName InstanceName, const FTracePath& TracePath) const;
+
+	/** Get/add perf data for instance and tracepath */
 	TSharedPtr<class FScriptPerfData> GetPerfDataByInstanceAndTracePath(FName InstanceName, const FTracePath& TracePath);
 
 	/** Get all instance perf data for the trace path, excluding the global blueprint data */
@@ -121,6 +133,29 @@ protected:
 
 };
 
+//////////////////////////////////////////////////////////////////////////
+// FScriptExecutionNodeParams
+
+struct KISMET_API FScriptExecNodeParams
+{
+	/** Node name */
+	FName NodeName;
+	/** Owning graph name */
+	FName OwningGraphName;
+	/** Node flags to describe the source graph node type */
+	uint32 NodeFlags;
+	/** Oberved object */
+	TWeakObjectPtr<const UObject> ObservedObject;
+	/** Display name for widget UI */
+	FText DisplayName;
+	/** Tooltip for widget UI */
+	FText Tooltip;
+	/** Icon color for widget UI */
+	FLinearColor IconColor;
+	/** Icon for widget UI */
+	FSlateBrush* Icon;
+};
+
 //////////////////////////////////////////////////////////////////////////
 // FScriptExecutionNode
 
@@ -131,6 +166,7 @@ class KISMET_API FScriptExecutionNode : public FScriptNodeExecLinkage, public FS
 public:
 
 	FScriptExecutionNode();
+	FScriptExecutionNode(const FScriptExecNodeParams& InitParams);
     virtual ~FScriptExecutionNode() {}
 
 	bool operator == (const FScriptExecutionNode& NodeIn) const;
@@ -138,21 +174,9 @@ public:
 	/** Get the node's name */
 	FName GetName() const { return NodeName; }
 
-	/** Set the node's name */
-	void SetName(FName NewName) { NodeName = NewName; }
-
-	/** Set the node's name by string */
-	void SetNameByString(const FString& NewName) { NodeName = FName(*NewName); }
-
 	/** Returns the owning graph name */
 	FName GetGraphName() const { return OwningGraphName; }
 
-	/** Sets the owning graph name */
-	void SetGraphName(FName GraphNameIn) { OwningGraphName = GraphNameIn; }
-
-	/** Set the node's flags */
-	void SetFlags(const uint32 NewFlags) { NodeFlags = NewFlags; }
-
 	/** Add to the node's flags */
 	void AddFlags(const uint32 NewFlags) { NodeFlags |= NewFlags; }
 
@@ -171,6 +195,9 @@ public:
 	/** Returns if this exec event represents the start of an event execution path */
 	bool IsEvent() const { return (NodeFlags & EScriptExecutionNodeFlags::Event) != 0U; }
 
+	/** Returns if this exec event represents the start of a custon event execution path */
+	bool IsCustomEvent() const { return (NodeFlags & EScriptExecutionNodeFlags::CustomEvent) != 0U; }
+	
 	/** Returns if this event is a function callsite event */
 	bool IsFunctionCallSite() const { return (NodeFlags & EScriptExecutionNodeFlags::FunctionCall) != 0U; }
 
@@ -180,12 +207,12 @@ public:
 	/** Returns if this event happened inside a macro instance */
 	bool IsMacroNode() const { return (NodeFlags & EScriptExecutionNodeFlags::MacroNode) != 0U; }
 
+	/** Returns if this node is also a pure node (i.e. no exec input pin) */
+	bool IsPureNode() const { return (NodeFlags & EScriptExecutionNodeFlags::PureNode) != 0U; }
+
 	/** Returns if this event potentially multiple exit sites */
 	bool IsBranch() const { return (NodeFlags & EScriptExecutionNodeFlags::BranchNode) != 0U; }
 
-	/** Sets the observed object context */
-	void SetObservedObject(const UObject* ObjectIn) { ObservedObject = ObjectIn; }
-
 	/** Gets the observed object context */
 	const UObject* GetObservedObject() { return ObservedObject.Get(); }
 
@@ -195,9 +222,6 @@ public:
 	/** Returns the display name for widget UI */
 	const FText& GetDisplayName() const { return DisplayName; }
 
-	/** Sets the display name for widget UI */
-	void SetDisplayName(const FText InDisplayName) { DisplayName = InDisplayName; }
-
 	/** Returns the tooltip for widget UI */
 	const FText& GetToolTipText() const { return Tooltip; }
 
@@ -213,9 +237,6 @@ public:
 	/** Returns the icon for widget UI */
 	const FSlateBrush* GetIcon() const { return Icon; }
 
-	/** Sets the icon for widget UI */
-	void SetIcon(FSlateBrush* InIcon) { Icon = InIcon; }
-
 	/** Returns the profiler heat color */
 	FLinearColor GetHeatColor(const FTracePath& TracePath) const;
 
@@ -225,6 +246,9 @@ public:
 	/** Sets the current expansion state for widget UI */
 	void SetExpanded(bool bIsExpanded) { bExpansionState = bIsExpanded; }
 
+	/** Sets the pure node script code range */
+	void SetPureNodeScriptCodeRange(FInt32Range InScriptCodeRange) { PureNodeScriptCodeRange = InScriptCodeRange; }
+
 	/** Return the linear execution path from this node */
 	void GetLinearExecutionPath(TArray<FLinearExecPath>& LinearExecutionNodes, const FTracePath& TracePath);
 
@@ -237,8 +261,14 @@ public:
 	/** Get all exec nodes */
 	virtual void GetAllExecNodes(TMap<FName, TSharedPtr<FScriptExecutionNode>>& ExecNodesOut);
 
+	/** Get all pure nodes associated with the given trace path */
+	virtual void GetAllPureNodes(TMap<int32, TSharedPtr<FScriptExecutionNode>>& PureNodesOut);
+
 protected:
 
+	/** Get all pure nodes - private implementation */
+	void GetAllPureNodes_Internal(TMap<int32, TSharedPtr<FScriptExecutionNode>>& PureNodesOut, const FInt32Range& ScriptCodeRange);
+
 	/** Get linear execution path - private implementation */
 	void GetLinearExecutionPath_Internal(FExecNodeFilter& Filter, TArray<FLinearExecPath>& LinearExecutionNodes, const FTracePath& TracePath);
 
@@ -265,7 +295,8 @@ protected:
 	FSlateBrush* Icon;
 	/** Expansion state */
 	bool bExpansionState;
-
+	/** Script code range for pure node linkage */
+	FInt32Range PureNodeScriptCodeRange;
 };
 
 //////////////////////////////////////////////////////////////////////////
@@ -275,6 +306,11 @@ class KISMET_API FScriptExecutionInstance : public FScriptExecutionNode
 {
 public:
 
+	FScriptExecutionInstance(const FScriptExecNodeParams& InitParams)
+		: FScriptExecutionNode(InitParams)
+	{
+	}
+
 	// FScriptExecutionNode
 	virtual void NavigateToObject() const override;
 	// ~FScriptExecutionNode
@@ -287,6 +323,11 @@ class KISMET_API FScriptExecutionBlueprint : public FScriptExecutionNode
 {
 public:
 
+	FScriptExecutionBlueprint(const FScriptExecNodeParams& InitParams)
+		: FScriptExecutionNode(InitParams)
+	{
+	}
+
 	/** Adds new blueprint instance node */
 	void AddInstance(TSharedPtr<FScriptExecutionNode> Instance) { Instances.Add(Instance); }
 
diff --git a/Engine/Source/Editor/Kismet/Public/Profiler/TracePath.h b/Engine/Source/Editor/Kismet/Public/Profiler/TracePath.h
--- a/Engine/Source/Editor/Kismet/Public/Profiler/TracePath.h
+++ b/Engine/Source/Editor/Kismet/Public/Profiler/TracePath.h
@@ -5,6 +5,8 @@
 //////////////////////////////////////////////////////////////////////////
 // FTracePath
 
+#define TRACEPATH_DEBUG 0
+
 class FTracePath
 {
 public:
@@ -16,8 +18,21 @@ public:
 
 	FTracePath(const FTracePath& TracePathIn)
 		: CachedCrc(TracePathIn.CachedCrc)
+#if TRACEPATH_DEBUG
+		, TraceStack(TracePathIn.TraceStack)
+#endif
+	{
+	}
+
+#if TRACEPATH_DEBUG
+	/** Copy operator */
+	FTracePath& operator = (const FTracePath& TracePathIn)
 	{
+		TraceStack = TracePathIn.TraceStack;
+		CachedCrc = TracePathIn.CachedCrc;
+		return *this;
 	}
+#endif
 
 	/** Cast to uint32 operator */
 	operator const uint32() const { return CachedCrc; }
@@ -25,17 +40,40 @@ public:
 	/** Add exit pin */
 	void AddExitPin(const int32 ExitPinScriptOffset)
 	{
+#if TRACEPATH_DEBUG
+		TraceStack.Add(ExitPinScriptOffset);
+#endif
 		CachedCrc = FCrc::MemCrc32(&ExitPinScriptOffset, 1, CachedCrc);
 	}
 
 	/** Reset */
 	void Reset()
 	{
+#if TRACEPATH_DEBUG
+		TraceStack.SetNum(0);
+#endif
 		CachedCrc = 0;
+
+	}
+
+#if TRACEPATH_DEBUG
+	const FString GetPathString() const
+	{
+		FString PathString(TEXT("<"));
+		for (auto Iter : TraceStack)
+		{
+			PathString = FString::Printf( TEXT( "%s:%i" ), *PathString, Iter);
+		}
+		PathString += TEXT(">");
+		return PathString;
 	}
+#endif
 
 private:
 
+#if TRACEPATH_DEBUG
+	TArray<int32> TraceStack;
+#endif
 	/** Cached crc */
 	uint32 CachedCrc;
 
diff --git a/Engine/Source/Editor/KismetCompiler/Private/KismetCompiler.cpp b/Engine/Source/Editor/KismetCompiler/Private/KismetCompiler.cpp
--- a/Engine/Source/Editor/KismetCompiler/Private/KismetCompiler.cpp
+++ b/Engine/Source/Editor/KismetCompiler/Private/KismetCompiler.cpp
@@ -1187,6 +1187,17 @@ void FKismetCompilerContext::PrecompileFunction(FKismetFunctionContext& Context)
 			ValidateSelfPinsInGraph(Context);
 			ValidateNoWildcardPinsInGraph(Context.SourceGraph);
 
+			for (int32 ChildIndex = 0; ChildIndex < Context.SourceGraph->Nodes.Num(); ++ChildIndex)
+			{
+				const UEdGraphNode* Node = Context.SourceGraph->Nodes[ChildIndex];
+				const int32 SavedErrorCount = MessageLog.NumErrors;
+				UK2Node_Event* SrcEventNode = Cast<UK2Node_Event>(Context.SourceGraph->Nodes[ChildIndex]);
+				if (bIsFullCompile || SrcEventNode)
+				{
+					ValidateNode(Node);
+				}
+			}
+
 			// Transforms
 			TransformNodes(Context);
 		}
@@ -1567,18 +1578,40 @@ void FKismetCompilerContext::CompileFunction(FKismetFunctionContext& Context)
 
 		if (IsNodePure(Node))
 		{
+			// For profiling purposes, find the statement that marks the function's entry point.
+			FBlueprintCompiledStatement* ProfilerStatement = nullptr;
+			TArray<FBlueprintCompiledStatement*>* SourceStatementList = Context.StatementsPerNode.Find(Node);
+			const bool bDidNodeGenerateCode = SourceStatementList != nullptr && SourceStatementList->Num() > 0;
+			if (bDidNodeGenerateCode)
+			{
+				for (auto Statement : *SourceStatementList)
+				{
+					if (Statement && Statement->Type == KCST_InstrumentedPureNodeEntry)
+					{
+						ProfilerStatement = Statement;
+						break;
+					}
+				}
+			}
+
 			// Push this node to the requirements list of any other nodes using it's outputs, if this node had any real impact
-			if (Context.DidNodeGenerateCode(Node) || bHasAntecedentPureNodes)
+			if (bDidNodeGenerateCode || bHasAntecedentPureNodes)
 			{
 				for (int32 PinIndex = 0; PinIndex < Node->Pins.Num(); ++PinIndex)
 				{
 					UEdGraphPin* Pin = Node->Pins[PinIndex];
-					if (Pin->Direction == EGPD_Output)
+					if (Pin->Direction == EGPD_Output && Pin->LinkedTo.Num() > 0)
 					{
-						for (int32 OutputIndex = 0; OutputIndex < Pin->LinkedTo.Num(); ++OutputIndex)
+						// Record the pure node output pin, since it's linked
+						if (ProfilerStatement)
+						{
+							ProfilerStatement->PureOutputContextArray.AddUnique(Pin);
+						}
+
+						for (UEdGraphPin* LinkedTo : Pin->LinkedTo)
 						{
-							UEdGraphNode* NodeUsingOutput = Pin->LinkedTo[OutputIndex]->GetOwningNode();
-							if (NodeUsingOutput != NULL)
+							UEdGraphNode* NodeUsingOutput = LinkedTo->GetOwningNode();
+							if (NodeUsingOutput != nullptr)
 							{
 								// Add this node, as well as other nodes this node depends on
 								TSet<UEdGraphNode*>& TargetNodesRequired = PureNodesNeeded.FindOrAdd(NodeUsingOutput);
@@ -1610,6 +1643,12 @@ void FKismetCompilerContext::CompileFunction(FKismetFunctionContext& Context)
 					OrderedInsertIntoArray(SortedPureNodes, SortKeyMap, *It);
 				}
 
+				if (Context.IsInstrumentationRequired())
+				{
+					FBlueprintCompiledStatement& PopState = Context.PrependStatementForNode(Node);
+					PopState.Type = KCST_InstrumentedStatePop;
+				}
+
 				// Inline their code
 				for (int32 i = 0; i < SortedPureNodes.Num(); ++i)
 				{
@@ -1617,6 +1656,12 @@ void FKismetCompilerContext::CompileFunction(FKismetFunctionContext& Context)
 
 					Context.CopyAndPrependStatements(Node, NodeToInline);
 				}
+
+				if (Context.IsInstrumentationRequired())
+				{
+					FBlueprintCompiledStatement& PushState = Context.PrependStatementForNode(Node);
+					PushState.Type = KCST_InstrumentedStatePush;
+				}
 			}
 
 			// Proceed to the next node
@@ -2596,16 +2641,8 @@ void FKismetCompilerContext::CreateAndProcessUbergraph()
 			// Validate all the nodes in the graph
 			for (int32 ChildIndex = 0; ChildIndex < ConsolidatedEventGraph->Nodes.Num(); ++ChildIndex)
 			{
-				const UEdGraphNode* Node = ConsolidatedEventGraph->Nodes[ChildIndex];
-				const int32 SavedErrorCount = MessageLog.NumErrors;
-				UK2Node_Event* SrcEventNode = Cast<UK2Node_Event>(ConsolidatedEventGraph->Nodes[ChildIndex]);
-				if (bIsFullCompile || SrcEventNode)
-				{
-					ValidateNode(Node);
-				}
-
-				// If the node didn't generate any errors then generate function stubs for event entry nodes etc.
-				if ((SavedErrorCount == MessageLog.NumErrors) && SrcEventNode)
+				UK2Node_Event* SrcEventNode = Cast<UK2Node_Event>(ConsolidatedEventGraph->Nodes[ChildIndex]);	
+				if (SrcEventNode)
 				{
 					CreateFunctionStubForEvent(SrcEventNode, Blueprint);
 				}
@@ -2941,41 +2978,29 @@ void FKismetCompilerContext::ProcessOneFunctionGraph(UEdGraph* SourceGraph, bool
 	// If a function in the graph cannot be overridden/placed as event make sure that it is not.
 	VerifyValidOverrideFunction(FunctionGraph);
 
-	// First do some cursory validation (pin types match, inputs to outputs, pins never point to their parent node, etc...)
-	// If this fails we don't proceed any further to avoid crashes or infinite loops
-	// When compiling only the skeleton class, we want the UFunction to be generated and processed so it contains all the local variables, this is unsafe to do during any other compilation mode
-	//
-	// NOTE: the order of this conditional check is intentional, and should not
-	//       be rearranged; we do NOT want ValidateGraphIsWellFormed() ran for 
-	//       skeleton-only compiles (that's why we have that check second) 
-	//       because it would most likely result in errors (the function hasn't
-	//       been added to the class yet, etc.)
-	if ((CompileOptions.CompileType == EKismetCompileType::SkeletonOnly) || ValidateGraphIsWellFormed(FunctionGraph))
-	{
-		FKismetFunctionContext& Context = *new (FunctionList)FKismetFunctionContext(MessageLog, Schema, NewClass, Blueprint, CompileOptions.DoesRequireCppCodeGeneration(), CompileOptions.IsInstrumentationActive());
-		Context.SourceGraph = FunctionGraph;
+	FKismetFunctionContext& Context = *new (FunctionList)FKismetFunctionContext(MessageLog, Schema, NewClass, Blueprint, CompileOptions.DoesRequireCppCodeGeneration(), CompileOptions.IsInstrumentationActive());
+	Context.SourceGraph = FunctionGraph;
 
-		if(FBlueprintEditorUtils::IsDelegateSignatureGraph(SourceGraph))
-		{
-			Context.SetDelegateSignatureName(SourceGraph->GetFName());
-		}
+	if(FBlueprintEditorUtils::IsDelegateSignatureGraph(SourceGraph))
+	{
+		Context.SetDelegateSignatureName(SourceGraph->GetFName());
+	}
 
-		// If this is an interface blueprint, mark the function contexts as stubs
-		if (FBlueprintEditorUtils::IsInterfaceBlueprint(Blueprint))
-		{
-			Context.MarkAsInterfaceStub();
-		}
+	// If this is an interface blueprint, mark the function contexts as stubs
+	if (FBlueprintEditorUtils::IsInterfaceBlueprint(Blueprint))
+	{
+		Context.MarkAsInterfaceStub();
+	}
 
-		bool bEnforceConstCorrectness = true;
-		if (FBlueprintEditorUtils::IsBlueprintConst(Blueprint) || Schema->IsConstFunctionGraph(Context.SourceGraph, &bEnforceConstCorrectness))
-		{
-			Context.MarkAsConstFunction(bEnforceConstCorrectness);
-		}
+	bool bEnforceConstCorrectness = true;
+	if (FBlueprintEditorUtils::IsBlueprintConst(Blueprint) || Schema->IsConstFunctionGraph(Context.SourceGraph, &bEnforceConstCorrectness))
+	{
+		Context.MarkAsConstFunction(bEnforceConstCorrectness);
+	}
 
-		if ( bInternalFunction )
-		{
-			Context.MarkAsInternalOrCppUseOnly();
-		}
+	if ( bInternalFunction )
+	{
+		Context.MarkAsInternalOrCppUseOnly();
 	}
 }
 
diff --git a/Engine/Source/Editor/KismetCompiler/Private/KismetCompilerMisc.cpp b/Engine/Source/Editor/KismetCompiler/Private/KismetCompilerMisc.cpp
--- a/Engine/Source/Editor/KismetCompiler/Private/KismetCompilerMisc.cpp
+++ b/Engine/Source/Editor/KismetCompiler/Private/KismetCompilerMisc.cpp
@@ -688,9 +688,20 @@ UEdGraphPin* FKismetCompilerUtilities::GenerateAssignmentNodes(class FKismetComp
 				}
 				else
 				{
-					// Move connection from the variable pin on the spawn node to the 'value' pin
-					CompilerContext.MovePinLinksToIntermediate(*OrgPin, *ValuePin);
-					SetVarNode->PinConnectionListChanged(ValuePin);
+					// For non-array struct pins that are not linked, transfer the pin type so that the node will expand an auto-ref that will assign the value by-ref.
+					if (OrgPin->PinType.bIsArray == false && OrgPin->PinType.PinCategory == Schema->PC_Struct && OrgPin->LinkedTo.Num() == 0)
+					{
+						ValuePin->PinType.PinCategory = OrgPin->PinType.PinCategory;
+						ValuePin->PinType.PinSubCategory = OrgPin->PinType.PinSubCategory;
+						ValuePin->PinType.PinSubCategoryObject = OrgPin->PinType.PinSubCategoryObject;
+						CompilerContext.MovePinLinksToIntermediate(*OrgPin, *ValuePin);
+					}
+					else
+					{
+						CompilerContext.MovePinLinksToIntermediate(*OrgPin, *ValuePin);
+						SetVarNode->PinConnectionListChanged(ValuePin);
+					}
+
 				}
 			}
 		}
diff --git a/Engine/Source/Editor/KismetCompiler/Private/KismetCompilerVMBackend.cpp b/Engine/Source/Editor/KismetCompiler/Private/KismetCompilerVMBackend.cpp
--- a/Engine/Source/Editor/KismetCompiler/Private/KismetCompilerVMBackend.cpp
+++ b/Engine/Source/Editor/KismetCompiler/Private/KismetCompilerVMBackend.cpp
@@ -225,6 +225,10 @@ private:
 	FKismetCompilerContext* CurrentCompilerContext;
 	FKismetFunctionContext* CurrentFunctionContext;
 
+	// Pure node count/starting offset (used for instrumentation)
+	int32 PureNodeEntryCount;
+	int32 PureNodeEntryStart;
+
 protected:
 	/**
 	 * This class is designed to be used like so to emit a bytecode context expression:
@@ -335,6 +339,8 @@ public:
 		, ReturnStatement(InReturnStatement)
 		, CurrentCompilerContext(nullptr)
 		, CurrentFunctionContext(nullptr)
+		, PureNodeEntryCount(0)
+		, PureNodeEntryStart(0)
 	{
 		VectorStruct = TBaseStructure<FVector>::Get();
 		RotatorStruct = TBaseStructure<FRotator>::Get();
@@ -1369,76 +1375,55 @@ public:
 		Writer.CommitSkip(PatchUpNeededAtOffset, Writer.ScriptBuffer.Num());
 	}
 
-	void EmitInstrumentation(FBlueprintCompiledStatement& Statement)
+	void EmitInstrumentation(FKismetCompilerContext& CompilerContext, FKismetFunctionContext& FunctionContext, FBlueprintCompiledStatement& Statement, UEdGraphNode* SourceNode)
 	{
-		int32 EventType = 0;
-		switch (Statement.Type)
+		int32 Offset = Writer.ScriptBuffer.Num();
+
+		if (Statement.Type == KCST_DebugSite)
+		{
+			Writer << EX_Tracepoint;
+		}
+		else if (Statement.Type == KCST_WireTraceSite)
 		{
+			Writer << EX_WireTracepoint;
+		}
+		else
+		{
+			int32 EventType = 0;
+			switch (Statement.Type)
+			{
 			case KCST_InstrumentedWireExit:			EventType = EScriptInstrumentation::NodeExit; break;
 			case KCST_InstrumentedWireEntry:		EventType = EScriptInstrumentation::NodeEntry; break;
 			case KCST_InstrumentedPureNodeEntry:	EventType = EScriptInstrumentation::PureNodeEntry; break;
 			case KCST_InstrumentedStatePush:		EventType = EScriptInstrumentation::PushState; break;
 			case KCST_InstrumentedStateRestore:		EventType = EScriptInstrumentation::RestoreState; break;
+			case KCST_InstrumentedStateReset:		EventType = EScriptInstrumentation::ResetState; break;
+			case KCST_InstrumentedStateSuspend:		EventType = EScriptInstrumentation::SuspendState; break;
 			case KCST_InstrumentedStatePop:			EventType = EScriptInstrumentation::PopState; break;
-		}
-		Writer << EX_InstrumentationEvent;
-		Writer << EventType;
-	}
-
-	void EmitArrayGetByRef(FBlueprintCompiledStatement& Statement)
-	{
-		Writer << EX_ArrayGetByRef;
-		// The array variable
-		EmitTerm(Statement.RHS[0]);
-		// The index to access in the array
-		EmitTerm(Statement.RHS[1], (UProperty*)(GetDefault<UIntProperty>()));
-	}
-
-	void PushReturnAddress(FBlueprintCompiledStatement& ReturnTarget)
-	{
-		Writer << EX_PushExecutionFlow;
-		CodeSkipSizeType PatchUpNeededAtOffset = Writer.EmitPlaceholderSkip();
-
-		JumpTargetFixupMap.Add(PatchUpNeededAtOffset, &ReturnTarget);
-	}
-
-	void CloseScript()
-	{
-		Writer << EX_EndOfScript;
-	}
-
-	virtual ~FScriptBuilderBase()
-	{
-	}
-
-	void GenerateCodeForStatement(FKismetCompilerContext& CompilerContext, FKismetFunctionContext& FunctionContext, FBlueprintCompiledStatement& Statement, UEdGraphNode* SourceNode)
-	{
-		TGuardValue<FKismetCompilerContext*> CompilerContextGuard(CurrentCompilerContext, &CompilerContext);
-		TGuardValue<FKismetFunctionContext*> FunctionContextGuard(CurrentFunctionContext, &FunctionContext);
+			}
+			Writer << EX_InstrumentationEvent;
+			Writer << EventType;
 
-		// Record the start of this statement in the bytecode if it's needed as a target label
-		if (Statement.bIsJumpTarget)
-		{
-			StatementLabelMap.Add(&Statement, Writer.ScriptBuffer.Num());
+			Offset += sizeof(int32);
 		}
 
-		// Generate bytecode for the statement
-		switch (Statement.Type)
+		if (Statement.Type != KCST_InstrumentedStateRestore)
 		{
-		case KCST_Nop:
-			Writer << EX_Nothing;
-			break;
-		case KCST_WireTraceSite:
+			TArray<UEdGraphPin*> PinContextArray(Statement.PureOutputContextArray);
+			if (Statement.ExecContext != nullptr)
+			{
+				PinContextArray.Add(Statement.ExecContext);
+			}
+
+			for (auto PinContext : PinContextArray)
 			{
-				UEdGraphPin const* TrueSourcePin = Cast<UEdGraphPin const>(FunctionContext.MessageLog.FindSourceObject(Statement.ExecContext));
+				UEdGraphPin const* TrueSourcePin = Cast<UEdGraphPin const>(FunctionContext.MessageLog.FindSourceObject(PinContext));
 				if (TrueSourcePin)
 				{
-					int32 Offset = Writer.ScriptBuffer.Num();
 					ClassBeingBuilt->GetDebugData().RegisterPinToCodeAssociation(TrueSourcePin, FunctionContext.Function, Offset);
 				}
 			}
-			// no break, continue down through KCST_DebugSite
-		case KCST_DebugSite:
+
 			if (SourceNode != NULL)
 			{
 				// Record where this NOP is
@@ -1474,11 +1459,75 @@ public:
 						CompilerContext.MessageLog.MacroSourceToMacroInstanceNodeMap.MultiFind(MacroSourceNode, MacroInstanceNodes);
 					}
 
-					int32 Offset = Writer.ScriptBuffer.Num();
 					ClassBeingBuilt->GetDebugData().RegisterNodeToCodeAssociation(TrueSourceNode, MacroSourceNode, MacroInstanceNodes, FunctionContext.Function, Offset, bBreakpointSite);
+
+					// Track pure node script code range for the current impure (exec) node
+					if (Statement.Type == KCST_InstrumentedPureNodeEntry)
+					{
+						if (PureNodeEntryCount == 0)
+						{
+							// Indicates the starting offset for this pure node call chain.
+							PureNodeEntryStart = Offset;
+						}
+
+						++PureNodeEntryCount;
+					}
+					else if (Statement.Type == KCST_InstrumentedWireEntry && PureNodeEntryCount > 0)
+					{
+						// Map script code range for the full set of pure node inputs feeding in to the current impure (exec) node at the current offset
+						ClassBeingBuilt->GetDebugData().RegisterPureNodeScriptCodeRange(TrueSourceNode, FunctionContext.Function, FInt32Range(PureNodeEntryStart, Offset));
+
+						// Reset pure node code range tracking.
+						PureNodeEntryCount = 0;
+						PureNodeEntryStart = 0;
+					}
 				}
 			}
-			Writer << ((Statement.Type == KCST_DebugSite) ? EX_Tracepoint : EX_WireTracepoint);
+		}
+	}
+	
+	void EmitArrayGetByRef(FBlueprintCompiledStatement& Statement)
+	{
+		Writer << EX_ArrayGetByRef;
+		// The array variable
+		EmitTerm(Statement.RHS[0]);
+		// The index to access in the array
+		EmitTerm(Statement.RHS[1], (UProperty*)(GetDefault<UIntProperty>()));
+	}
+
+	void PushReturnAddress(FBlueprintCompiledStatement& ReturnTarget)
+	{
+		Writer << EX_PushExecutionFlow;
+		CodeSkipSizeType PatchUpNeededAtOffset = Writer.EmitPlaceholderSkip();
+
+		JumpTargetFixupMap.Add(PatchUpNeededAtOffset, &ReturnTarget);
+	}
+
+	void CloseScript()
+	{
+		Writer << EX_EndOfScript;
+	}
+
+	virtual ~FScriptBuilderBase()
+	{
+	}
+
+	void GenerateCodeForStatement(FKismetCompilerContext& CompilerContext, FKismetFunctionContext& FunctionContext, FBlueprintCompiledStatement& Statement, UEdGraphNode* SourceNode)
+	{
+		TGuardValue<FKismetCompilerContext*> CompilerContextGuard(CurrentCompilerContext, &CompilerContext);
+		TGuardValue<FKismetFunctionContext*> FunctionContextGuard(CurrentFunctionContext, &FunctionContext);
+
+		// Record the start of this statement in the bytecode if it's needed as a target label
+		if (Statement.bIsJumpTarget)
+		{
+			StatementLabelMap.Add(&Statement, Writer.ScriptBuffer.Num());
+		}
+
+		// Generate bytecode for the statement
+		switch (Statement.Type)
+		{
+		case KCST_Nop:
+			Writer << EX_Nothing;
 			break;
 		case KCST_CallFunction:
 			EmitFunctionCall(CompilerContext, FunctionContext, Statement, SourceNode);
@@ -1548,45 +1597,18 @@ public:
 		case KCST_SwitchValue:
 			EmitSwitchValue(Statement);
 			break;
-		case KCST_InstrumentedWireExit:
+		case KCST_DebugSite:
+		case KCST_WireTraceSite:
 		case KCST_InstrumentedWireEntry:
+		case KCST_InstrumentedWireExit:
 		case KCST_InstrumentedStatePush:
+		case KCST_InstrumentedStateReset:
+		case KCST_InstrumentedStateSuspend:
 		case KCST_InstrumentedStatePop:
-			{
-				UEdGraphPin const* TrueSourcePin = Cast<UEdGraphPin const>(FunctionContext.MessageLog.FindSourceObject(Statement.ExecContext));
-				if (TrueSourcePin)
-				{
-					int32 Offset = Writer.ScriptBuffer.Num() + sizeof(int32);
-					ClassBeingBuilt->GetDebugData().RegisterPinToCodeAssociation(TrueSourcePin, FunctionContext.Function, Offset);
-				}
-				if (SourceNode != NULL)
-				{
-					// Record where this NOP is
-					UEdGraphNode* TrueSourceNode = Cast<UEdGraphNode>(FunctionContext.MessageLog.FindSourceObject(SourceNode));
-					if (TrueSourceNode)
-					{
-						// If this is a debug site for an expanded macro instruction, there should also be a macro source node associated with it
-						UEdGraphNode* MacroSourceNode = Cast<UEdGraphNode>(CompilerContext.MessageLog.FinalNodeBackToMacroSourceMap.FindSourceObject(SourceNode));
-						if (MacroSourceNode == SourceNode)
-						{
-							// The function above will return the given node if not found in the map. In that case there is no associated source macro node, so we clear it.
-							MacroSourceNode = NULL;
-						}
-	
-						TArray<TWeakObjectPtr<UEdGraphNode>> MacroInstanceNodes;
-						int32 Offset = Writer.ScriptBuffer.Num() + sizeof(int32);
-						ClassBeingBuilt->GetDebugData().RegisterNodeToCodeAssociation(TrueSourceNode, MacroSourceNode, MacroInstanceNodes, FunctionContext.Function, Offset, false);
-					}
-				}
-			}
-		case KCST_InstrumentedPureNodeEntry:
 		case KCST_InstrumentedStateRestore:
-			{
-				// Emit Statement
-				EmitInstrumentation(Statement);
-				break;
-			}
-
+		case KCST_InstrumentedPureNodeEntry:
+			EmitInstrumentation(CompilerContext, FunctionContext, Statement, SourceNode);
+			break;
 		case KCST_ArrayGetByRef:
 			EmitArrayGetByRef(Statement);
 			break;
diff --git a/Engine/Source/Editor/KismetCompiler/Public/BlueprintCompiledStatement.h b/Engine/Source/Editor/KismetCompiler/Public/BlueprintCompiledStatement.h
--- a/Engine/Source/Editor/KismetCompiler/Public/BlueprintCompiledStatement.h
+++ b/Engine/Source/Editor/KismetCompiler/Public/BlueprintCompiledStatement.h
@@ -43,6 +43,8 @@ enum EKismetCompiledStatementType
 	KCST_InstrumentedWireExit,		// Instrumented wiretrace exit
 	KCST_InstrumentedStatePush,		// Instrumented state push
 	KCST_InstrumentedStateRestore,	// Instrumented state restore
+	KCST_InstrumentedStateReset,	// Instrumented state reset
+	KCST_InstrumentedStateSuspend,	// Instrumented state suspend
 	KCST_InstrumentedStatePop,		// Instrumented state pop
 	//
 	KCST_ArrayGetByRef,
@@ -97,6 +99,9 @@ struct FBlueprintCompiledStatement
 	// Exec pin about to execute (KCST_WireTraceSite)
 	class UEdGraphPin* ExecContext;
 
+	// Pure node output pin(s) linked to exec node input pins (KCST_InstrumentedPureNodeEntry)
+	TArray<class UEdGraphPin*> PureOutputContextArray;
+
 	// Comment text
 	FString Comment;
 };
diff --git a/Engine/Source/Editor/KismetCompiler/Public/KismetCompiledFunctionContext.h b/Engine/Source/Editor/KismetCompiler/Public/KismetCompiledFunctionContext.h
--- a/Engine/Source/Editor/KismetCompiler/Public/KismetCompiledFunctionContext.h
+++ b/Engine/Source/Editor/KismetCompiler/Public/KismetCompiledFunctionContext.h
@@ -234,6 +234,17 @@ public:
 		return Blueprint;
 	}
 
+	FBlueprintCompiledStatement& PrependStatementForNode(UEdGraphNode* Node)
+	{
+		FBlueprintCompiledStatement* Result = new FBlueprintCompiledStatement();
+		AllGeneratedStatements.Add(Result);
+
+		TArray<FBlueprintCompiledStatement*>& StatementList = StatementsPerNode.FindOrAdd(Node);
+		StatementList.Insert(Result, 0);
+
+		return *Result;
+	}
+
 	/** Enqueue a statement to be executed when the specified Node is triggered */
 	FBlueprintCompiledStatement& AppendStatementForNode(UEdGraphNode* Node)
 	{
diff --git a/Engine/Source/Editor/UMGEditor/Private/Details/SWidgetDetailsView.cpp b/Engine/Source/Editor/UMGEditor/Private/Details/SWidgetDetailsView.cpp
--- a/Engine/Source/Editor/UMGEditor/Private/Details/SWidgetDetailsView.cpp
+++ b/Engine/Source/Editor/UMGEditor/Private/Details/SWidgetDetailsView.cpp
@@ -219,6 +219,8 @@ void SWidgetDetailsView::OnEditorSelectionChanged()
 	{
 		for ( FWidgetReference& WidgetRef : SelectedWidgets )
 		{
+			// Edit actions will go directly to the preview widget, changes will be
+			// propagated to the template via SWidgetDetailsView::NotifyPostChange
 			SelectedObjects.Add(WidgetRef.GetPreview());
 		}
 	}
diff --git a/Engine/Source/Editor/UnrealEd/Private/ScriptDisassembler.cpp b/Engine/Source/Editor/UnrealEd/Private/ScriptDisassembler.cpp
--- a/Engine/Source/Editor/UnrealEd/Private/ScriptDisassembler.cpp
+++ b/Engine/Source/Editor/UnrealEd/Private/ScriptDisassembler.cpp
@@ -928,6 +928,9 @@ void FKismetBytecodeDisassembler::ProcessCommon(int32& ScriptIndex, EExprToken O
 		{
 			Ar.Logf(TEXT("%s $%X: Array Get-by-Ref Index"), *Indents, (int32)Opcode);
 			AddIndent();
+			SerializeExpr(ScriptIndex);
+			SerializeExpr(ScriptIndex);
+			DropIndent();
 			break;
 		}
 	default:
diff --git a/Engine/Source/Editor/UnrealEd/Private/Settings/SettingsClasses.cpp b/Engine/Source/Editor/UnrealEd/Private/Settings/SettingsClasses.cpp
--- a/Engine/Source/Editor/UnrealEd/Private/Settings/SettingsClasses.cpp
+++ b/Engine/Source/Editor/UnrealEd/Private/Settings/SettingsClasses.cpp
@@ -15,6 +15,7 @@
 
 #include "SourceCodeNavigation.h"
 
+#include "Developer/BlueprintProfiler/Public/BlueprintProfilerModule.h"
 #include "ScriptPerfData.h"
 
 #define LOCTEXT_NAMESPACE "SettingsClasses"
@@ -93,6 +94,18 @@ void UEditorExperimentalSettings::PostEditChangeProperty( struct FPropertyChange
 			FModuleManager::Get().LoadModule(TEXT("EnvironmentQueryEditor"));
 		}
 	}
+	else if (Name == FName(TEXT("bBlueprintPerformanceAnalysisTools")))
+	{
+		if (!bBlueprintPerformanceAnalysisTools)
+		{
+			IBlueprintProfilerInterface* ProfilerInterface = FModuleManager::GetModulePtr<IBlueprintProfilerInterface>("BlueprintProfiler");
+			if (ProfilerInterface && ProfilerInterface->IsProfilerEnabled())
+			{
+				// Force Profiler off
+				ProfilerInterface->ToggleProfilingCapture();
+			}
+		}
+	}
 	else if (Name == FName(TEXT("BlueprintProfilerRecentSampleBias")))
 	{
 		FScriptPerfData::SetRecentSampleBias(BlueprintProfilerRecentSampleBias);
diff --git a/Engine/Source/Programs/UnrealBuildTool/Configuration/UEBuildModule.cs b/Engine/Source/Programs/UnrealBuildTool/Configuration/UEBuildModule.cs
--- a/Engine/Source/Programs/UnrealBuildTool/Configuration/UEBuildModule.cs
+++ b/Engine/Source/Programs/UnrealBuildTool/Configuration/UEBuildModule.cs
@@ -801,6 +801,13 @@ namespace UnrealBuildTool
 		/// </summary>
 		public bool bSkipDefinitionsForCompileEnvironment = false;
 
+        public int FindNumberOfGeneratedCppFiles()
+        {
+           return ((null == GeneratedCodeDirectory) || !GeneratedCodeDirectory.Exists()) ? 0
+                : (GeneratedCodeDirectory.EnumerateFileReferences("*.generated.*.cpp", SearchOption.AllDirectories).Count() 
+                 + GeneratedCodeDirectory.EnumerateFileReferences("*.generated.cpp", SearchOption.AllDirectories).Count());
+        }
+
 		/// <summary>
 		/// Categorizes source files into per-extension buckets
 		/// </summary>
diff --git a/Engine/Source/Programs/UnrealBuildTool/Configuration/UEBuildTarget.cs b/Engine/Source/Programs/UnrealBuildTool/Configuration/UEBuildTarget.cs
--- a/Engine/Source/Programs/UnrealBuildTool/Configuration/UEBuildTarget.cs
+++ b/Engine/Source/Programs/UnrealBuildTool/Configuration/UEBuildTarget.cs
@@ -2856,10 +2856,23 @@ namespace UnrealBuildTool
 
 			// Fill out the body of the function with the empty function calls. This is what causes the static libraries to be considered relevant
 			List<UEBuildModule> DependencyModules = ExecutableBinary.GetAllDependencyModules(bIncludeDynamicallyLoaded: false, bForceCircular: false);
-			foreach (string ModuleName in DependencyModules.OfType<UEBuildModuleCPP>().Where(CPPModule => CPPModule.AutoGenerateCppInfo != null).Select(CPPModule => CPPModule.Name).Distinct())
+            HashSet<string> AlreadyAddedEmptyLinkFunctions = new HashSet<string>();
+            foreach (UEBuildModuleCPP BuildModuleCPP in DependencyModules.OfType<UEBuildModuleCPP>().Where(CPPModule => CPPModule.AutoGenerateCppInfo != null))
 			{
-				Result.Add("    extern void EmptyLinkFunctionForGeneratedCode" + ModuleName + "();");
-				Result.Add("    EmptyLinkFunctionForGeneratedCode" + ModuleName + "();");
+                int NumGeneratedCppFilesWithTheFunction = BuildModuleCPP.FindNumberOfGeneratedCppFiles();
+                if(NumGeneratedCppFilesWithTheFunction == 0)
+                {
+                    Result.Add("    //" + BuildModuleCPP.Name + " has no generated files, path: " + BuildModuleCPP.GeneratedCodeDirectory.ToString());
+                }
+                for (int FileIdx = 1; FileIdx <= NumGeneratedCppFilesWithTheFunction; ++FileIdx)
+                {
+                    string FunctionName = "EmptyLinkFunctionForGeneratedCode" + FileIdx + BuildModuleCPP.Name;
+                    if (AlreadyAddedEmptyLinkFunctions.Add(FunctionName))
+                    {
+                        Result.Add("    extern void " + FunctionName + "();");
+                        Result.Add("    " + FunctionName + "();");
+                    }
+                }
 			}
 			foreach (string DependencyModuleName in PrivateDependencyModuleNames)
 			{
diff --git a/Engine/Source/Programs/UnrealHeaderTool/Private/CodeGenerator.cpp b/Engine/Source/Programs/UnrealHeaderTool/Private/CodeGenerator.cpp
--- a/Engine/Source/Programs/UnrealHeaderTool/Private/CodeGenerator.cpp
+++ b/Engine/Source/Programs/UnrealHeaderTool/Private/CodeGenerator.cpp
@@ -1495,7 +1495,6 @@ void FNativeClassHeaderGenerator::ExportNativeGeneratedInitCode(FClass* Class, F
 			{
 				GeneratedClassRegisterFunctionText.Logf(TEXT("\t\t\t\t%s(CastChecked<UDynamicClass>(OuterClass));\n"), *(*CustomDynamicClassInitializationMD));
 			}
-			GeneratedClassRegisterFunctionText.Logf(TEXT("\t\t\t\tOuterClass->GetDefaultObject();\r\n"));
 		}
 
 		GeneratedClassRegisterFunctionText.Logf(TEXT("\t\t\t}\r\n"));
@@ -5295,7 +5294,6 @@ void FNativeClassHeaderGenerator::ExportGeneratedCPP()
 	FUHTStringBuilder GeneratedCPPClassesIncludes;
 	FUHTStringBuilder GeneratedCPPEpilogue;
 	FUHTStringBuilder GeneratedCPPText;
-	FUHTStringBuilder GeneratedLinkerFixupFunction;
 
 	TArray<FUHTStringBuilder> GeneratedCPPFiles;
 
@@ -5328,8 +5326,6 @@ void FNativeClassHeaderGenerator::ExportGeneratedCPP()
 	// Write out our include to the .dep.h file
 	GeneratedCPPClassesIncludes.Logf(TEXT("#include \"%s\"") LINE_TERMINATOR, *FPaths::GetCleanFilename(DepHeaderPathname));
 
-	GeneratedLinkerFixupFunction.Logf(TEXT("void EmptyLinkFunctionForGeneratedCode%s() {}") LINE_TERMINATOR, *ModuleInfo->Name);
-
 	{
 		// Autogenerate names (alphabetically sorted).
 		ReferencedNames.KeySort( TLess<FName>() );
@@ -5383,7 +5379,8 @@ void FNativeClassHeaderGenerator::ExportGeneratedCPP()
 		}
 
 		FString CppPath = ModuleInfo->GeneratedCPPFilenameBase + (GeneratedFunctionBodyTextSplit.Num() > 1 ? *FString::Printf(TEXT(".%d.cpp"), FileIdx + 1) : TEXT(".cpp"));
-		SaveHeaderIfChanged(*CppPath, *(GeneratedCPPPreamble + ModulePCHInclude + GeneratedCPPClassesIncludes + DisableDeprecationWarnings + ((FileIdx > 0) ? FString() : GeneratedLinkerFixupFunction) + FileText + GeneratedCPPEpilogue + EnableDeprecationWarnings));
+		const FString GeneratedLinkerFixupFunction = FString::Printf(TEXT("void EmptyLinkFunctionForGeneratedCode%d%s() {}") LINE_TERMINATOR, FileIdx + 1, *ModuleInfo->Name);
+		SaveHeaderIfChanged(*CppPath, *(GeneratedCPPPreamble + ModulePCHInclude + GeneratedCPPClassesIncludes + DisableDeprecationWarnings + GeneratedLinkerFixupFunction + FileText + GeneratedCPPEpilogue + EnableDeprecationWarnings));
 
 		if (GeneratedFunctionBodyTextSplit.Num() > 1)
 		{
diff --git a/Engine/Source/Runtime/CoreUObject/Private/Serialization/AsyncLoading.cpp b/Engine/Source/Runtime/CoreUObject/Private/Serialization/AsyncLoading.cpp
--- a/Engine/Source/Runtime/CoreUObject/Private/Serialization/AsyncLoading.cpp
+++ b/Engine/Source/Runtime/CoreUObject/Private/Serialization/AsyncLoading.cpp
@@ -1635,7 +1635,8 @@ EAsyncPackageState::Type FAsyncPackage::LoadImports()
 
 		// Handle circular dependencies - try to find existing packages.
 		UPackage* ExistingPackage = dynamic_cast<UPackage*>(StaticFindObjectFast(UPackage::StaticClass(), nullptr, ImportPackageFName, true));
-		if (ExistingPackage && !ExistingPackage->HasAnyPackageFlags(PKG_CompiledIn) && !ExistingPackage->bHasBeenFullyLoaded)//!ExistingPackage->HasAnyFlags(RF_WasLoaded))
+		if (ExistingPackage  && !ExistingPackage->bHasBeenFullyLoaded 
+			&& (!ExistingPackage->HasAnyPackageFlags(PKG_CompiledIn) || GetConvertedDynamicPackageNameToTypeName().Contains(ImportPackageFName)))//!ExistingPackage->HasAnyFlags(RF_WasLoaded))
 		{
 			// The import package already exists. Check if it's currently being streamed as well. If so, make sure
 			// we add all dependencies that don't yet have linkers created otherwise we risk that if the current package
diff --git a/Engine/Source/Runtime/CoreUObject/Private/UObject/Class.cpp b/Engine/Source/Runtime/CoreUObject/Private/UObject/Class.cpp
--- a/Engine/Source/Runtime/CoreUObject/Private/UObject/Class.cpp
+++ b/Engine/Source/Runtime/CoreUObject/Private/UObject/Class.cpp
@@ -4972,6 +4972,12 @@ void UDynamicClass::AddReferencedObjects(UObject* InThis, FReferenceCollector& C
 	Super::AddReferencedObjects(This, Collector);
 }
 
+void UDynamicClass::PostLoad()
+{
+	Super::PostLoad();
+	GetDefaultObject(true);
+}
+
 void UDynamicClass::PurgeClass(bool bRecompilingOnLoad)
 {
 	Super::PurgeClass(bRecompilingOnLoad);
diff --git a/Engine/Source/Runtime/CoreUObject/Private/UObject/GarbageCollection.cpp b/Engine/Source/Runtime/CoreUObject/Private/UObject/GarbageCollection.cpp
--- a/Engine/Source/Runtime/CoreUObject/Private/UObject/GarbageCollection.cpp
+++ b/Engine/Source/Runtime/CoreUObject/Private/UObject/GarbageCollection.cpp
@@ -1181,6 +1181,12 @@ void CollectGarbageInternal(EObjectFlags KeepFlags, bool bPerformFullPurge)
 			bool bShouldAssert = false;
 			FUObjectItem* ObjectItem = *It;
 			UObject* Object = (UObject*)ObjectItem->Object;
+
+			if (Object->HasAnyFlags(RF_ClassDefaultObject) && Object->GetName().Contains(TEXT("RecastFilter_UseDefaultArea")))
+			{
+				printf("");
+			}
+
 			// Don't require UGCObjectReferencer's references to adhere to the assumptions.
 			// Although we want the referencer itself to sit in the disregard for gc set, most of the objects
 			// it's referencing will not be in the root set.
@@ -1192,19 +1198,25 @@ void CollectGarbageInternal(EObjectFlags KeepFlags, bool bPerformFullPurge)
 				ObjectReferenceCollector.FindReferences( Object );
 
 				// Iterate over referenced objects, finding bad ones.
-				for( int32 ReferenceIndex=0; ReferenceIndex<CollectedReferences.Num(); ReferenceIndex++ )
+				for (int32 ReferenceIndex = 0; ReferenceIndex < CollectedReferences.Num(); ReferenceIndex++)
 				{
 					UObject* ReferencedObject = CollectedReferences[ReferenceIndex];
-					if (ReferencedObject && 
-						!(ReferencedObject->IsRooted() || 
-						  UObjectArray.IsDisregardForGC(ReferencedObject) || 
-							UObjectArray.ObjectToObjectItem(ReferencedObject)->GetOwnerIndex() ||
-							UObjectArray.ObjectToObjectItem(ReferencedObject)->HasAnyFlags(EInternalObjectFlags::ClusterRoot)))
+
+
+					if (ReferencedObject)
 					{
-						UE_LOG(LogGarbage, Warning, TEXT("Disregard for GC object %s referencing %s which is not part of root set"),
-							*Object->GetFullName(),
-							*ReferencedObject->GetFullName());						
-						bShouldAssert = true;
+						bool bIsRooted = ReferencedObject->IsRooted();
+						bool bIsDisregard = UObjectArray.IsDisregardForGC(ReferencedObject);
+						int32 OwnerIndex = UObjectArray.ObjectToObjectItem(ReferencedObject)->GetOwnerIndex();
+						bool bHasClusterRoot = UObjectArray.ObjectToObjectItem(ReferencedObject)->HasAnyFlags(EInternalObjectFlags::ClusterRoot);
+
+						if (!(bIsRooted || bIsDisregard || OwnerIndex || bHasClusterRoot))
+						{
+							UE_LOG(LogGarbage, Warning, TEXT("Disregard for GC object %s referencing %s which is not part of root set"),
+								*Object->GetFullName(),
+								*ReferencedObject->GetFullName());
+							bShouldAssert = true;
+						}
 					}
 				}
 			}
diff --git a/Engine/Source/Runtime/CoreUObject/Private/UObject/ScriptCore.cpp b/Engine/Source/Runtime/CoreUObject/Private/UObject/ScriptCore.cpp
--- a/Engine/Source/Runtime/CoreUObject/Private/UObject/ScriptCore.cpp
+++ b/Engine/Source/Runtime/CoreUObject/Private/UObject/ScriptCore.cpp
@@ -1128,7 +1128,8 @@ void UObject::ProcessEvent( UFunction* Function, void* Parms )
 #if !(UE_BUILD_SHIPPING || UE_BUILD_TEST)
 	if (GetClass()->HasInstrumentation())
 	{
-		EScriptInstrumentationEvent EventInstrumentationInfo(EScriptInstrumentation::Event, this, Function->GetFName());
+		const EScriptInstrumentation::Type EventType = Function->HasAnyFunctionFlags(FUNC_Event|FUNC_BlueprintEvent) ? EScriptInstrumentation::Event : EScriptInstrumentation::ResumeEvent;
+		EScriptInstrumentationEvent EventInstrumentationInfo(EventType, this, Function->GetFName());
 		FBlueprintCoreDelegates::InstrumentScriptEvent(EventInstrumentationInfo);
 	}
 #endif
diff --git a/Engine/Source/Runtime/CoreUObject/Private/UObject/UObjectArray.cpp b/Engine/Source/Runtime/CoreUObject/Private/UObject/UObjectArray.cpp
--- a/Engine/Source/Runtime/CoreUObject/Private/UObject/UObjectArray.cpp
+++ b/Engine/Source/Runtime/CoreUObject/Private/UObject/UObjectArray.cpp
@@ -10,6 +10,16 @@ DEFINE_LOG_CATEGORY_STATIC(LogUObjectArray, Log, All);
 
 TMap<int32, FUObjectCluster*> GUObjectClusters;
 
+FUObjectArray::FUObjectArray()
+: ObjFirstGCIndex(0)
+, ObjLastNonGCIndex(INDEX_NONE)
+, MaxObjectsNotConsideredByGC(0)
+, OpenForDisregardForGC(!IS_PROGRAM)
+, MasterSerialNumber(START_SERIAL_NUMBER)
+{
+	FCoreDelegates::GetObjectArrayForDebugVisualizersDelegate().BindStatic(GetObjectArrayForDebugVisualizers);
+}
+
 void FUObjectArray::AllocateObjectPool(int32 InMaxUObjects, int32 InMaxObjectsNotConsideredByGC)
 {
 	check(IsInGameThread());
@@ -48,17 +58,17 @@ void FUObjectArray::CloseDisregardForGC()
 	// assembles the token reference stream at this point. This is required for class objects that are
 	// not taken into account for garbage collection but have instances that are.
 
-	// Workaround for Visual Studio 2013 analyzer bug. Using a temporary directly in the range-for
-	// errors if the analyzer is enabled.
-	TObjectRange<UClass> Range;
-	for (UClass* Class : Range)
+	for (FRawObjectIterator It(false); It; ++It) // GetDefaultObject can create a new class, that need to be handled as well, so we cannot use TObjectIterator
 	{
-		// Force the default object to be created.
-		Class->GetDefaultObject(); // Force the default object to be constructed if it isn't already
-		// Assemble reference token stream for garbage collection/ RTGC.
-		if (!Class->HasAnyClassFlags(CLASS_TokenStreamAssembled))
+		if (UClass* Class = Cast<UClass>((UObject*)(It->Object)))
 		{
-			Class->AssembleReferenceTokenStream();
+			// Force the default object to be created.
+			Class->GetDefaultObject(); // Force the default object to be constructed if it isn't already
+			// Assemble reference token stream for garbage collection/ RTGC.
+			if (!Class->HasAnyClassFlags(CLASS_TokenStreamAssembled))
+			{
+				Class->AssembleReferenceTokenStream();
+			}
 		}
 	}
 
@@ -102,6 +112,17 @@ void FUObjectArray::CloseDisregardForGC()
 	GIsInitialLoad = false;
 }
 
+void FUObjectArray::DisableDisregardForGC()
+{
+	MaxObjectsNotConsideredByGC = 0;
+	ObjFirstGCIndex = 0;
+	ObjLastNonGCIndex = INDEX_NONE;
+	if (IsOpenForDisregardForGC())
+	{
+		CloseDisregardForGC();
+	}
+}
+
 void FUObjectArray::AllocateUObjectIndex(UObjectBase* Object, bool bMergingThreads /*= false*/)
 {
 	int32 Index = INDEX_NONE;
diff --git a/Engine/Source/Runtime/CoreUObject/Public/UObject/Class.h b/Engine/Source/Runtime/CoreUObject/Public/UObject/Class.h
--- a/Engine/Source/Runtime/CoreUObject/Public/UObject/Class.h
+++ b/Engine/Source/Runtime/CoreUObject/Public/UObject/Class.h
@@ -2507,6 +2507,7 @@ public:
 
 	// UObject interface.
 	static void AddReferencedObjects(UObject* InThis, FReferenceCollector& Collector);
+	virtual void PostLoad() override;
 
 	// UClass interface
 	virtual void PurgeClass(bool bRecompilingOnLoad) override;
diff --git a/Engine/Source/Runtime/CoreUObject/Public/UObject/Script.h b/Engine/Source/Runtime/CoreUObject/Public/UObject/Script.h
--- a/Engine/Source/Runtime/CoreUObject/Public/UObject/Script.h
+++ b/Engine/Source/Runtime/CoreUObject/Public/UObject/Script.h
@@ -304,16 +304,15 @@ namespace EScriptInstrumentation
 		Class = 0,
 		Instance,
 		Event,
-		Function,
-		Branch,
-		Macro,
+		ResumeEvent,
 		PureNodeEntry,
 		NodeEntry,
 		NodeExit,
 		PushState,
 		RestoreState,
+		ResetState,
+		SuspendState,
 		PopState,
-		NodePin,
 		Stop
 	};
 }
diff --git a/Engine/Source/Runtime/CoreUObject/Public/UObject/UObjectArray.h b/Engine/Source/Runtime/CoreUObject/Public/UObject/UObjectArray.h
--- a/Engine/Source/Runtime/CoreUObject/Public/UObject/UObjectArray.h
+++ b/Engine/Source/Runtime/CoreUObject/Public/UObject/UObjectArray.h
@@ -323,14 +323,7 @@ public:
 	/**
 	 * Constructor, initializes to no permanent object pool
 	 */
-	FUObjectArray() :
-		ObjFirstGCIndex(0),
-		ObjLastNonGCIndex(INDEX_NONE),
-		OpenForDisregardForGC(true),
-		MasterSerialNumber(START_SERIAL_NUMBER)
-	{
-		FCoreDelegates::GetObjectArrayForDebugVisualizersDelegate().BindStatic(GetObjectArrayForDebugVisualizers);
-	}
+	FUObjectArray();
 
 	/**
 	 * Allocates and initializes the permanent object pool
@@ -341,13 +334,10 @@ public:
 	void AllocateObjectPool(int32 MaxUObjects, int32 MaxObjectsNotConsideredByGC);
 
 	/**
-	 * Disables the disregard for GC optimization. Commandlets can't use it.
+	 * Disables the disregard for GC optimization.
 	 *
 	 */
-	void DisableDisregardForGC()
-	{
-		ObjFirstGCIndex = 0;
-	}
+	void DisableDisregardForGC();
 
 	/**
 	* If there's enough slack in the disregard pool, we can re-open it and keep adding objects to it
diff --git a/Engine/Source/Runtime/Engine/Classes/Animation/AnimClassInterface.h b/Engine/Source/Runtime/Engine/Classes/Animation/AnimClassInterface.h
--- a/Engine/Source/Runtime/Engine/Classes/Animation/AnimClassInterface.h
+++ b/Engine/Source/Runtime/Engine/Classes/Animation/AnimClassInterface.h
@@ -39,6 +39,7 @@ public:
 		}
 		if (auto DynamicClass = Cast<UDynamicClass>(InClass))
 		{
+			DynamicClass->GetDefaultObject(true);
 			return CastChecked<IAnimClassInterface>(DynamicClass->AnimClassImplementation, ECastCheckedType::NullAllowed);
 		}
 		return nullptr;
diff --git a/Engine/Source/Runtime/Engine/Classes/Engine/BlueprintGeneratedClass.h b/Engine/Source/Runtime/Engine/Classes/Engine/BlueprintGeneratedClass.h
--- a/Engine/Source/Runtime/Engine/Classes/Engine/BlueprintGeneratedClass.h
+++ b/Engine/Source/Runtime/Engine/Classes/Engine/BlueprintGeneratedClass.h
@@ -52,8 +52,14 @@ public:
 	// Reverse map from code offset to macro instance node(s)
 	TMultiMap< int32, TWeakObjectPtr<UEdGraphNode> > LineNumberToMacroInstanceNodeMap;
 
-	// Reverse map from code offset to exec pin
-	TMap< int32, TWeakObjectPtr<UEdGraphPin> > LineNumberToPinMap;
+	// Reverse map from code offset to source pin
+	TMap< int32, TWeakObjectPtr<UEdGraphPin> > LineNumberToSourcePinMap;
+
+	// Reverse map from source pin to mapped code offset(s)
+	TMultiMap< TWeakObjectPtr<UEdGraphPin>, int32 > SourcePinToLineNumbersMap;
+
+	// Map from source node (impure) to pure node script code range
+	TMap< TWeakObjectPtr<UEdGraphNode>, FInt32Range > PureNodeScriptCodeRangeMap;
 
 public:
 	FDebuggingInfoForSingleFunction()
@@ -266,17 +272,61 @@ public:
 		}
 	}
 
-	// Finds the macro source node associated with the code location Function+CodeOffset, or nullptr if there isn't one
-	UEdGraphPin* FindExecPinFromCodeLocation(UFunction* Function, int32 CodeOffset) const
+	// Finds the source pin associated with the code location Function+CodeOffset, or nullptr if there isn't one
+	UEdGraphPin* FindSourcePinFromCodeLocation(UFunction* Function, int32 CodeOffset) const
 	{
 		if (const FDebuggingInfoForSingleFunction* pFuncInfo = PerFunctionLineNumbers.Find(Function))
 		{
-			return pFuncInfo->LineNumberToPinMap.FindRef(CodeOffset).Get();
+			return pFuncInfo->LineNumberToSourcePinMap.FindRef(CodeOffset).Get();
 		}
 
 		return nullptr;
 	}
 
+	// Finds all code locations (Function+CodeOffset) associated with the source pin
+	void FindAllCodeLocationsFromSourcePin(UEdGraphPin const* SourcePin, UFunction* InFunction, TArray<int32>& OutPinToCodeAssociations) const
+	{
+		OutPinToCodeAssociations.Empty();
+
+		if (const FDebuggingInfoForSingleFunction* pFuncInfo = PerFunctionLineNumbers.Find(InFunction))
+		{
+			pFuncInfo->SourcePinToLineNumbersMap.MultiFind(SourcePin, OutPinToCodeAssociations, true);
+		}
+	}
+
+	// Finds the first code location (Function+CodeOffset) associated with the source pin within the given range, or INDEX_NONE if there isn't one
+	int32 FindCodeLocationFromSourcePin(UEdGraphPin const* SourcePin, UFunction* InFunction, FInt32Range InRange = FInt32Range()) const
+	{
+		TArray<int32> PinToCodeAssociations;
+		FindAllCodeLocationsFromSourcePin(SourcePin, InFunction, PinToCodeAssociations);
+
+		for (int32 i = 0; i < PinToCodeAssociations.Num(); ++i)
+		{
+			if (InRange.Contains(PinToCodeAssociations[i]))
+			{
+				return PinToCodeAssociations[i];
+			}
+		}
+
+		return INDEX_NONE;
+	}
+
+	// Finds the pure node script code range associated with the [impure] source node, or FInt32Range(INDEX_NONE) if there is no existing association
+	FInt32Range FindPureNodeScriptCodeRangeFromSourceNode(const UEdGraphNode* SourceNode, UFunction* InFunction) const
+	{
+		FInt32Range Result = FInt32Range(INDEX_NONE);
+
+		if (const FDebuggingInfoForSingleFunction* DebugInfoPtr = PerFunctionLineNumbers.Find(InFunction))
+		{
+			if (const FInt32Range* ValuePtr = DebugInfoPtr->PureNodeScriptCodeRangeMap.Find(SourceNode))
+			{
+				Result = *ValuePtr;
+			}
+		}
+
+		return Result;
+	}
+
 	// Finds the breakpoint injection site(s) in bytecode if any were associated with the given node
 	void FindBreakpointInjectionSites(UEdGraphNode* Node, TArray<uint8*>& InstallSites) const
 	{
@@ -335,21 +385,17 @@ public:
 		}
 	}
 
-	void RegisterPinToCodeAssociation(UEdGraphPin const* ExecPin, UFunction* InFunction, int32 CodeOffset)
+	void RegisterPureNodeScriptCodeRange(UEdGraphNode* TrueSourceNode, UFunction* InFunction, FInt32Range InPureNodeScriptCodeRange)
 	{
 		FDebuggingInfoForSingleFunction& PerFuncInfo = PerFunctionLineNumbers.FindOrAdd(InFunction);
-		PerFuncInfo.LineNumberToPinMap.Add(CodeOffset, ExecPin);
+		PerFuncInfo.PureNodeScriptCodeRangeMap.Add(TrueSourceNode, InPureNodeScriptCodeRange);
 	}
 
-	int32 FindPinToCodeAssociation(UEdGraphPin const* ExecPin, UFunction* InFunction)
+	void RegisterPinToCodeAssociation(UEdGraphPin const* SourcePin, UFunction* InFunction, int32 CodeOffset)
 	{
 		FDebuggingInfoForSingleFunction& PerFuncInfo = PerFunctionLineNumbers.FindOrAdd(InFunction);
-		int32 Result = INDEX_NONE;
-		if (const int32* CodeOffset = PerFuncInfo.LineNumberToPinMap.FindKey(ExecPin))
-		{
-			Result = *CodeOffset;
-		}
-		return Result;
+		PerFuncInfo.LineNumberToSourcePinMap.Add(CodeOffset, SourcePin);
+		PerFuncInfo.SourcePinToLineNumbersMap.Add(SourcePin, CodeOffset);
 	}
 
 	// Registers an association between an object (pin or node typically) and an associated class member property
diff --git a/Engine/Source/Runtime/Engine/Classes/Kismet/KismetMathLibrary.h b/Engine/Source/Runtime/Engine/Classes/Kismet/KismetMathLibrary.h
--- a/Engine/Source/Runtime/Engine/Classes/Kismet/KismetMathLibrary.h
+++ b/Engine/Source/Runtime/Engine/Classes/Kismet/KismetMathLibrary.h
@@ -921,6 +921,20 @@ class ENGINE_API UKismetMathLibrary : public UBlueprintFunctionLibrary
 	UFUNCTION(BlueprintPure, meta=(DisplayName = "LinearColor * Float", CompactNodeTitle = "*", Keywords = "* multiply"), Category="Math|Color")
 	static FLinearColor Multiply_LinearColorFloat(FLinearColor A, float B);
 
+	//
+	// Plane functions.
+	//
+	
+	/** 
+	* Creates a plane with a facing direction of Normal at the given Point
+	* 
+	* @param Point	A point on the plane
+	* @param Normal  The Normal of the plane at Point
+	* @return Plane instance
+	*/
+	UFUNCTION(BlueprintPure, Category = "Math|Plane", meta=(Keywords="make plane"))
+	static FPlane MakePlaneFromPointAndNormal(FVector Point, FVector Normal);
+	
 	//
 	// DateTime functions.
 	//
diff --git a/Engine/Source/Runtime/Engine/Classes/Kismet/KismetSystemLibrary.h b/Engine/Source/Runtime/Engine/Classes/Kismet/KismetSystemLibrary.h
--- a/Engine/Source/Runtime/Engine/Classes/Kismet/KismetSystemLibrary.h
+++ b/Engine/Source/Runtime/Engine/Classes/Kismet/KismetSystemLibrary.h
@@ -618,7 +618,7 @@ class ENGINE_API UKismetSystemLibrary : public UBlueprintFunctionLibrary
 	}
 
 	/** Set a custom structure property by name */
-	UFUNCTION(BlueprintCallable, CustomThunk, meta = (BlueprintInternalUseOnly = "true", CustomStructureParam = "Value"))
+	UFUNCTION(BlueprintCallable, CustomThunk, meta = (BlueprintInternalUseOnly = "true", CustomStructureParam = "Value", AutoCreateRefTerm = "Value"))
 	static void SetStructurePropertyByName(UObject* Object, FName PropertyName, const FGenericStruct& Value);
 
 	static void Generic_SetStructurePropertyByName(UObject* OwnerObject, FName StructPropertyName, const void* SrcStructAddr);
diff --git a/Engine/Source/Runtime/Engine/Private/KismetFunctionLibrary.cpp b/Engine/Source/Runtime/Engine/Private/KismetFunctionLibrary.cpp
--- a/Engine/Source/Runtime/Engine/Private/KismetFunctionLibrary.cpp
+++ b/Engine/Source/Runtime/Engine/Private/KismetFunctionLibrary.cpp
@@ -2,6 +2,12 @@
 
 #include "EnginePrivate.h"
 
+#if WITH_EDITOR
+#include "Settings/LevelEditorPlaySettings.h" // for GetRunUnderOneProcess()
+#endif
+
+DEFINE_LOG_CATEGORY_STATIC(LogBlueprintFuncLibrary, Log, All);
+
 UBlueprintFunctionLibrary::UBlueprintFunctionLibrary(const FObjectInitializer& ObjectInitializer)
 	: Super(ObjectInitializer)
 {
@@ -9,21 +15,99 @@ UBlueprintFunctionLibrary::UBlueprintFunctionLibrary(const FObjectInitializer& O
 
 int32 UBlueprintFunctionLibrary::GetFunctionCallspace(UFunction* Function, void* Parameters, FFrame* Stack)
 {
-	// This is really awkward and totally sucks.
-	//
-	// This used to call GEngine->GetNetMode. With multiple netdrivers and worlds, this no longer makes sense.
-	//
-	// Without an actor or world to give us context, we don't know what netdriver to look at for the netmode.
-	//
-	// We could ban KismetAuthorityOnly for static functions, and force the authority check to be done manually in the function.
-	// This is awkward though since we lose the automatic 'server only' icon in kismet, which is pretty important.
-	//
-	// For now, we will use globabls to indicate what world we are in, and essentially default to GWorld and GameNetDriver to 
-	// see whether we should absorb the function call.
-
-	const bool Absorb = (Function->HasAllFunctionFlags(FUNC_BlueprintAuthorityOnly) && GEngine->ShouldAbsorbAuthorityOnlyEvent()) ||
-						(Function->HasAllFunctionFlags(FUNC_BlueprintCosmetic) && GEngine->ShouldAbsorbCosmeticOnlyEvent());
-	return Absorb ? FunctionCallspace::Absorbed : FunctionCallspace::Local;
+	const bool bIsAuthoritativeFunc = Function->HasAnyFunctionFlags(FUNC_BlueprintAuthorityOnly);
+	const bool bIsCosmeticFunc      = Function->HasAnyFunctionFlags(FUNC_BlueprintCosmetic);
+
+	bool bAbsorbFunctionCall = false;
+	if (bIsAuthoritativeFunc || bIsCosmeticFunc)
+	{
+#if !WITH_EDITOR
+		// without an actor or world to give us context, we don't know what net-
+		// driver to look at for "net-mode" - these functions will loop over all 
+		// available worlds and make a judgment off of that
+		//
+		// this is relatively fast, but can be inaccurate in PIE when we're 
+		// running a single process running a dedicated server as well
+		// (hence, why we do all the heavy logic below to find a world context)
+		const bool bAbsorbCosmeticCalls  = GEngine->ShouldAbsorbCosmeticOnlyEvent();
+		const bool bAbsorbAuthorityCalls = GEngine->ShouldAbsorbAuthorityOnlyEvent();
+#else
+		bool bAbsorbCosmeticCalls  = false;
+		bool bAbsorbAuthorityCalls = false;
+
+		const ULevelEditorPlaySettings* PlayInSettings = GetDefault<ULevelEditorPlaySettings>();
+
+		bool bIsOnePIEProcessEnabled = false;
+		// if we're set to run our simulation under one process, then we run the 
+		// risk of having conflicting world net-modes; in that scenario, the 
+		// ShouldAbsorb...Event() calls are unreliable - you want to explicitly 
+		// check which world the cosmetic/authoritative call would run in
+		if (PlayInSettings->GetRunUnderOneProcess(bIsOnePIEProcessEnabled) && bIsOnePIEProcessEnabled)
+		{
+			UObject* WorldContextObj = nullptr;
+			if (Stack != nullptr)
+			{
+//				const FName WorldContextTag = TEXT("WorldContext");
+
+				if (Stack->Object && Stack->Object->ImplementsGetWorld())
+				{
+					WorldContextObj = Stack->Object;
+				}
+// 				else if (Function->HasMetaData(WorldContextTag))
+// 				{
+// 					// @TODO: we could try to read the world context value off 
+// 					//        the stack; though I am not entirely comfortable 
+// 					//        without looking into it more - this would mean 
+// 					//        that the script feeding the function parameters 
+// 					//        would be ran twice (once here, to preview the 
+// 					//        value, and then again if we decide to proceed with
+// 					//        calling the function though normal stack execution),
+// 					//        and side effects will be ran more than expected.
+// 				}
+			}
+
+			if (WorldContextObj != nullptr)
+			{
+				UWorld* WorldContext = WorldContextObj->GetWorld();
+				if (ensure(WorldContext != nullptr))
+				{
+					ENetMode WorldNetMode = WorldContext->GetNetMode();
+					bAbsorbCosmeticCalls  = (WorldNetMode == NM_DedicatedServer);
+					bAbsorbAuthorityCalls = (WorldNetMode == NM_Client);
+				}
+			}
+			else
+			{
+				// fallback to old (possibly incorrect) behavior 
+				bAbsorbCosmeticCalls  = GEngine->ShouldAbsorbCosmeticOnlyEvent();
+				bAbsorbAuthorityCalls = GEngine->ShouldAbsorbAuthorityOnlyEvent();
+
+				// @TODO: consider warning, if we're dropping cosmetic events 
+				//        in a single process PIE session, using a dedicated 
+				//        server (we don't know if this is supposed to be 
+				//        running for the server world, or the client one 
+				//
+				//        conversely, do we want to warn when authoritative 
+				//        calls are dropped? in case they're supposed to be 
+				//        running on that dedicated server world?
+			}
+		}
+		else
+		{
+			// presumably, these calls are sufficient if there aren't worlds 
+			// with differing net-modes running within this one process... these
+			// will just loop through the running worlds, and look at their 
+			// net-mode
+			bAbsorbCosmeticCalls  = GEngine->ShouldAbsorbCosmeticOnlyEvent();
+			bAbsorbAuthorityCalls = GEngine->ShouldAbsorbAuthorityOnlyEvent();
+		}
+#endif //
+
+		bAbsorbFunctionCall = (bIsAuthoritativeFunc && bAbsorbAuthorityCalls) || (bIsCosmeticFunc && bAbsorbCosmeticCalls);
+	}
+	// else, there is no reason to check for authoritative/cosmetic absorption
+	
+	return bAbsorbFunctionCall ? FunctionCallspace::Absorbed : FunctionCallspace::Local;
 }
 
 FStringAssetReference UBlueprintFunctionLibrary::Generic_MakeStringAssetReference(FFrame& Stack, const FString& AssetLongPathname)
diff --git a/Engine/Source/Runtime/Engine/Private/KismetMathLibrary.cpp b/Engine/Source/Runtime/Engine/Private/KismetMathLibrary.cpp
--- a/Engine/Source/Runtime/Engine/Private/KismetMathLibrary.cpp
+++ b/Engine/Source/Runtime/Engine/Private/KismetMathLibrary.cpp
@@ -1304,7 +1304,13 @@ bool UKismetMathLibrary::ClassIsChildOf(TSubclassOf<class UObject> TestClass, TS
 	return ((*ParentClass != NULL) && (*TestClass != NULL)) ? (*TestClass)->IsChildOf(*ParentClass) : false;
 }
 
-
+/* Plane functions
+ *****************************************************************************/
+FPlane UKismetMathLibrary::MakePlaneFromPointAndNormal(FVector Point, FVector Normal)
+{ 
+	return FPlane(Point, Normal.GetSafeNormal());
+}
+ 
 /* DateTime functions
  *****************************************************************************/
 FDateTime UKismetMathLibrary::MakeDateTime(int32 Year, int32 Month, int32 Day, int32 Hour, int32 Minute, int32 Second, int32 Millisecond)
diff --git a/Engine/Source/Runtime/Launch/Private/LaunchEngineLoop.cpp b/Engine/Source/Runtime/Launch/Private/LaunchEngineLoop.cpp
--- a/Engine/Source/Runtime/Launch/Private/LaunchEngineLoop.cpp
+++ b/Engine/Source/Runtime/Launch/Private/LaunchEngineLoop.cpp
@@ -1580,7 +1580,10 @@ int32 FEngineLoop::PreInit( const TCHAR* CmdLine )
 		return 1;
 	}
 
-	GUObjectArray.CloseDisregardForGC();
+	if (GUObjectArray.IsOpenForDisregardForGC())
+	{
+		GUObjectArray.CloseDisregardForGC();
+	}
 
 	if (IOnlineSubsystem::IsLoaded())
 	{
