Copying //UE4/Dev-Blueprints to Dev-Main (//UE4/Dev-Main)

==========================
MAJOR FEATURES + CHANGES
==========================

Change 2910968 on 2016/03/15 by [REDACTED_PERSON]

	PR #2152: Fix for UE-27652 (Contributed by [REDACTED_PERSON])

Change 2911052 on 2016/03/15 by [REDACTED_PERSON]

	Moving Pin tooltip generating into GetPinHoverText to correct some issues with hovertext being stale. Motivation was PR #2095: Keep ResultPin Tooltip up to date on SpawnActor and CreateObject Nodes (Contributed by [REDACTED_PERSON])

Change 2912261 on 2016/03/16 by [REDACTED_PERSON]

	Fixed DynamicClass loading.

Change 2912719 on 2016/03/17 by [REDACTED_PERSON]

	[UE-28450] Cooked component template data now properly handles array property values having one or more items that differ from the default object.

	change summary:
	- modified FBlueprintEditorUtils::BuildComponentInstancingData() to generate additional data for individual array property values that differ from the component's class default object at cook time.
	- modified FBlueprintCookedComponentInstancingData::BuildCachedPropertyList() to parse cooked array property data at load time and generate a SubPropertyList representing the set of changed array property values.
	- modified UArrayProperty::SerializeItem() to handle an explicit custom property list when specified as part of the FArchive.
	- modified AActor::CreateComponentFromTemplateData() to set the 'PPF_Duplicate' PortFlag in order to emulate expected behavior that would otherwise occur when SDO is used for component instancing.
	- modified UActorComponent::Serialize() to not set 'bHasBeenCreated' on load for component template objects.

Change 2912749 on 2016/03/17 by [REDACTED_PERSON]

	[UE-28450] CIS fix

Change 2912811 on 2016/03/17 by [REDACTED_PERSON]

	[UEBP-112] Fix for a minor issue w/ branch nodes not including pure node chains in profiler tree view. Also reversed order of pure node chain display in tree view.

Change 2912890 on 2016/03/17 by [REDACTED_PERSON]

	Blueprint C++ Conversion:
	- Constructor handles private (inaccessible) member variables.
	- no ensure failed when a component was removed in a static actor instance.

Change 2913115 on 2016/03/17 by [REDACTED_PERSON]

	[UEBP-112] Fix for a potential profiler crash when processing a pure node event with a script code offset that also maps to an impure exec node in a different function context.

	(forgot to include this in the previous CL)

Change 2913266 on 2016/03/17 by [REDACTED_PERSON]

	Blueprint C++ Conversion:
	- fixed super:: call for BP native event.

Change 2913313 on 2016/03/17 by [REDACTED_PERSON]

	[UEBP-112] Fix reversed trace paths on pure nodes which caused missing pure timing data.

	(introduced in CL# 2912811)

Change 2915001 on 2016/03/18 by [REDACTED_PERSON]

	Blueprint C++ Conversion: Fixed color construction

Change 2915306 on 2016/03/18 by [REDACTED_PERSON]

	Fix /initProperties() optimization when PostLoad() changes a CDO property value away from the default value.

Change 2916140 on 2016/03/20 by [REDACTED_PERSON]

	Improved array property handling for both InitProperties() and cooked component data optimizations.

	Note: This should allow for proper runtime handling of array properties with an Inner that is potentially also an array property.

	change summary:
	- added FObjectInitializer::InitArrayPropertyFromCustomList() to assist with initializing instances from array properties with an explicit sub-property list.
	- added FBlueprintCookedComponentInstancingData::BuildCachedArrayPropertyList() to assist with recursively building a sub-property list for array deltas.
	- added FBlueprintGeneratedClass::BuildCustomArrayPropertyListForPostConstruction() to assist with recursively building a sub-property list for array deltas.

Change 2916640 on 2016/03/21 by [REDACTED_PERSON]

	Support for tunnelling graphs in the blueprint profiler
	#UEBP-183 - Tunneling graph support
	#Proj Kismet, BlueprintProfiler

	- This removes the need to filter stat updates to prevent cyclic references ( now tunnel nodes are better described in terms of i/o pins )
	- Moves to the more complete function based mapping of graphs.
	- Precursor to final macro code

Change 2916902 on 2016/03/21 by [REDACTED_PERSON]

	The graph panel will now properly fade in splines when highlighting them.

	Added dynamic tick registering that will invalidate the graph panel, allowing it to update the current display per tick until the action is complete.

Change 2917087 on 2016/03/21 by [REDACTED_PERSON]

	CIS mac fix

Change 2917433 on 2016/03/21 by [REDACTED_PERSON]

	PR #2163: Fix Bug for "IsDataOnlyBlueprint" Error (Contributed by [REDACTED_PERSON])

Change 2917499 on 2016/03/21 by [REDACTED_PERSON]

	Correcting an expectation that custom collision mappings were ordered (by collision channel). Now sorting the list before we translate them into trace types.

Change 2919538 on 2016/03/23 by [REDACTED_PERSON]

diff --git a/Engine/Plugins/Experimental/LiveEditor/Source/LiveEditor/Private/K2Node_LiveEditObject.cpp b/Engine/Plugins/Experimental/LiveEditor/Source/LiveEditor/Private/K2Node_LiveEditObject.cpp
--- a/Engine/Plugins/Experimental/LiveEditor/Source/LiveEditor/Private/K2Node_LiveEditObject.cpp
+++ b/Engine/Plugins/Experimental/LiveEditor/Source/LiveEditor/Private/K2Node_LiveEditObject.cpp
@@ -244,6 +244,7 @@ void UK2Node_LiveEditObject::ReallocatePinsDuringReconstruction(TArray<UEdGraphP
 	{
 		CreatePinsForClass(UseSpawnClass);
 	}
+	RestoreSplitPins(OldPins);
 }
 
 FNodeHandlingFunctor* UK2Node_LiveEditObject::CreateNodeHandler(FKismetCompilerContext& CompilerContext) const
diff --git a/Engine/Source/Developer/BlueprintCompilerCppBackend/Private/BlueprintCompilerCppBackend.cpp b/Engine/Source/Developer/BlueprintCompilerCppBackend/Private/BlueprintCompilerCppBackend.cpp
--- a/Engine/Source/Developer/BlueprintCompilerCppBackend/Private/BlueprintCompilerCppBackend.cpp
+++ b/Engine/Source/Developer/BlueprintCompilerCppBackend/Private/BlueprintCompilerCppBackend.cpp
@@ -564,7 +564,7 @@ FString FBlueprintCompilerCppBackend::EmitCallStatmentInner(FEmitterLocalContext
 		}
 		Result += FunctionToCallOriginalName;
 
-		if (Statement.bIsParentContext && FEmitHelper::ShouldHandleAsNativeEvent(Statement.FunctionToCall))
+		if (Statement.bIsParentContext && FEmitHelper::ShouldHandleAsNativeEvent(Statement.FunctionToCall, false))
 		{
 			ensure(!bCallOnDifferentObject);
 			Result += TEXT("_Implementation");
@@ -606,7 +606,7 @@ FString FBlueprintCompilerCppBackend::TermToText(FEmitterLocalContext& EmitterCo
 	const FString PSC_Self(TEXT("self"));
 	if (Term->bIsLiteral)
 	{
-		return FEmitHelper::LiteralTerm(EmitterContext, Term->Type, Term->Name, Term->ObjectLiteral);
+		return FEmitHelper::LiteralTerm(EmitterContext, Term->Type, Term->Name, Term->ObjectLiteral, &Term->TextLiteral);
 	}
 	else if (Term->InlineGeneratedParameter)
 	{
diff --git a/Engine/Source/Developer/BlueprintCompilerCppBackend/Private/BlueprintCompilerCppBackendBase.cpp b/Engine/Source/Developer/BlueprintCompilerCppBackend/Private/BlueprintCompilerCppBackendBase.cpp
--- a/Engine/Source/Developer/BlueprintCompilerCppBackend/Private/BlueprintCompilerCppBackendBase.cpp
+++ b/Engine/Source/Developer/BlueprintCompilerCppBackend/Private/BlueprintCompilerCppBackendBase.cpp
@@ -398,7 +398,9 @@ void FBlueprintCompilerCppBackendBase::ConstructFunction(FKismetFunctionContext&
 					const uint32 ExportFlags = EPropertyExportCPPFlags::CPPF_CustomTypeName | EPropertyExportCPPFlags::CPPF_BlueprintCppBackend | EPropertyExportCPPFlags::CPPF_NoConst;
 					const FString ActualArg = EmitterContext.ExportCppDeclaration(Property, EExportedDeclaration::Parameter, ExportFlags, false);
 					const FString NoConstType = EmitterContext.ExportCppDeclaration(Property, EExportedDeclaration::Parameter, ExportFlags, true);
-					EmitterContext.AddLine(FString::Printf(TEXT("%s = const_cast<%s>(%s__const);"), *ActualArg, *NoConstType, *FEmitHelper::GetCppName(Property)));
+					const FString TypeDefName = FString(TEXT("T")) + EmitterContext.GenerateUniqueLocalName();
+					EmitterContext.AddLine(FString::Printf(TEXT("typedef %s %s;"), *NoConstType, *TypeDefName));
+					EmitterContext.AddLine(FString::Printf(TEXT("%s = *const_cast<%s *>(&%s__const);"), *ActualArg, *TypeDefName, *FEmitHelper::GetCppName(Property)));
 				}
 			}
 			DeclareLocalVariables(EmitterContext, LocalVariables);
@@ -660,10 +662,11 @@ FString FBlueprintCompilerCppBackendBase::GenerateCodeFromEnum(UUserDefinedEnum*
 	check(SourceEnum);
 	FCodeText Header;
 	Header.AddLine(TEXT("#pragma once"));
+	const FString EnumCppName = *FEmitHelper::GetCppName(SourceEnum);
 	// use GetBaseFilename() so that we can coordinate #includes and filenames
 	Header.AddLine(FString::Printf(TEXT("#include \"%s.generated.h\""), *FEmitHelper::GetBaseFilename(SourceEnum)));
 	Header.AddLine(FString::Printf(TEXT("UENUM(BlueprintType, %s )"), *FEmitHelper::ReplaceConvertedMetaData(SourceEnum)));
-	Header.AddLine(FString::Printf(TEXT("enum class %s  : uint8"), *FEmitHelper::GetCppName(SourceEnum)));
+	Header.AddLine(FString::Printf(TEXT("enum class %s  : uint8"), *EnumCppName));
 	Header.AddLine(TEXT("{"));
 	Header.IncreaseIndent();
 	for (int32 Index = 0; Index < SourceEnum->NumEnums(); ++Index)
@@ -672,11 +675,32 @@ FString FBlueprintCompilerCppBackendBase::GenerateCodeFromEnum(UUserDefinedEnum*
 		const int32 ElemValue = Index;
 
 		const FString& DisplayNameMD = SourceEnum->GetMetaData(TEXT("DisplayName"), ElemValue);// TODO: value or index?
-		const FString Meta = DisplayNameMD.IsEmpty() ? FString() : FString::Printf(TEXT("UMETA(DisplayName = \"%s\")"), *DisplayNameMD);
+		const FString Meta = DisplayNameMD.IsEmpty() ? FString() : FString::Printf(TEXT("UMETA(DisplayName = \"%s\")"), *DisplayNameMD.ReplaceCharWithEscapedChar());
 		Header.AddLine(FString::Printf(TEXT("%s = %d %s,"), *ElemName, ElemValue, *Meta));
 	}
 	Header.DecreaseIndent();
 	Header.AddLine(TEXT("};"));
+	/*
+	Header.AddLine(FString::Printf(TEXT("FString %s__GetUserFriendlyName(int32 InValue)"), *EnumCppName));
+	Header.AddLine(TEXT("{"));
+	Header.IncreaseIndent();
+
+	Header.AddLine(TEXT("switch(InValue)"));
+	Header.AddLine(TEXT("{"));
+	Header.IncreaseIndent();
+	for (int32 Index = 0; Index < SourceEnum->NumEnums(); ++Index)
+	{
+		const FString DisplayNameStr = SourceEnum->GetEnumText(Index).ToString().ReplaceCharWithEscapedChar();
+		Header.AddLine(FString::Printf(TEXT("case %s::%s: return FString(TEXT(\"%s\"));"), *EnumCppName, *SourceEnum->GetEnumName(Index), *DisplayNameStr));
+	}
+	Header.DecreaseIndent();
+	Header.AddLine(TEXT("};"));
+
+	Header.AddLine(TEXT("ensure(false);"));
+	Header.AddLine(TEXT("return FString();"));
+	Header.DecreaseIndent();
+	Header.AddLine(TEXT("};"));
+	*/
 	return Header.Result;
 }
 
diff --git a/Engine/Source/Developer/BlueprintCompilerCppBackend/Private/BlueprintCompilerCppBackendUtils.cpp b/Engine/Source/Developer/BlueprintCompilerCppBackend/Private/BlueprintCompilerCppBackendUtils.cpp
--- a/Engine/Source/Developer/BlueprintCompilerCppBackend/Private/BlueprintCompilerCppBackendUtils.cpp
+++ b/Engine/Source/Developer/BlueprintCompilerCppBackend/Private/BlueprintCompilerCppBackendUtils.cpp
@@ -813,7 +813,7 @@ void FEmitHelper::EmitLifetimeReplicatedPropsImpl(FEmitterLocalContext& EmitterC
 	}
 }
 
-FString FEmitHelper::LiteralTerm(FEmitterLocalContext& EmitterContext, const FEdGraphPinType& Type, const FString& CustomValue, UObject* LiteralObject)
+FString FEmitHelper::LiteralTerm(FEmitterLocalContext& EmitterContext, const FEdGraphPinType& Type, const FString& CustomValue, UObject* LiteralObject, const FText* OptionalTextLiteral)
 {
 	auto Schema = GetDefault<UEdGraphSchema_K2>();
 
@@ -823,6 +823,11 @@ FString FEmitHelper::LiteralTerm(FEmitterLocalContext& EmitterContext, const FEd
 	}
 	else if (UEdGraphSchema_K2::PC_Text == Type.PinCategory)
 	{
+		ensure(OptionalTextLiteral);
+		if (OptionalTextLiteral)
+		{
+			return UTextProperty::GenerateCppCodeForTextValue(*OptionalTextLiteral, FString());
+		}
 		return FString::Printf(TEXT("FText::FromString(%s)"), *UStrProperty::ExportCppHardcodedText(CustomValue, EmitterContext.DefaultTarget->Indent));
 	}
 	else if (UEdGraphSchema_K2::PC_Float == Type.PinCategory)
@@ -1140,7 +1145,7 @@ FString FEmitHelper::DefaultValue(FEmitterLocalContext& EmitterContext, const FE
 		return FString::Printf(TEXT("%s{}"), *InnerTypeName);
 	}
 
-	return LiteralTerm(EmitterContext, Type, FString(), nullptr);
+	return LiteralTerm(EmitterContext, Type, FString(), nullptr, &FText::GetEmpty());
 }
 
 UFunction* FEmitHelper::GetOriginalFunction(UFunction* Function)
@@ -1169,12 +1174,12 @@ UFunction* FEmitHelper::GetOriginalFunction(UFunction* Function)
 	return Function;
 }
 
-bool FEmitHelper::ShouldHandleAsNativeEvent(UFunction* Function)
+bool FEmitHelper::ShouldHandleAsNativeEvent(UFunction* Function, bool bOnlyIfOverridden)
 {
 	check(Function);
 	auto OriginalFunction = FEmitHelper::GetOriginalFunction(Function);
 	check(OriginalFunction);
-	if (OriginalFunction != Function)
+	if (!bOnlyIfOverridden || (OriginalFunction != Function))
 	{
 		const uint32 FlagsToCheckMask = EFunctionFlags::FUNC_Event | EFunctionFlags::FUNC_BlueprintEvent | EFunctionFlags::FUNC_Native;
 		const uint32 FlagsToCheck = OriginalFunction->FunctionFlags & FlagsToCheckMask;
diff --git a/Engine/Source/Developer/BlueprintCompilerCppBackend/Private/BlueprintCompilerCppBackendUtils.h b/Engine/Source/Developer/BlueprintCompilerCppBackend/Private/BlueprintCompilerCppBackendUtils.h
--- a/Engine/Source/Developer/BlueprintCompilerCppBackend/Private/BlueprintCompilerCppBackendUtils.h
+++ b/Engine/Source/Developer/BlueprintCompilerCppBackend/Private/BlueprintCompilerCppBackendUtils.h
@@ -229,7 +229,7 @@ struct FEmitHelper
 
 	static void EmitLifetimeReplicatedPropsImpl(FEmitterLocalContext& EmitterContext);
 
-	static FString LiteralTerm(FEmitterLocalContext& EmitterContext, const FEdGraphPinType& Type, const FString& CustomValue, UObject* LiteralObject);
+	static FString LiteralTerm(FEmitterLocalContext& EmitterContext, const FEdGraphPinType& Type, const FString& CustomValue, UObject* LiteralObject, const FText* OptionalTextLiteral = nullptr);
 
 	static FString DefaultValue(FEmitterLocalContext& EmitterContext, const FEdGraphPinType& Type);
 
@@ -237,7 +237,7 @@ struct FEmitHelper
 
 	static UFunction* GetOriginalFunction(UFunction* Function);
 
-	static bool ShouldHandleAsNativeEvent(UFunction* Function);
+	static bool ShouldHandleAsNativeEvent(UFunction* Function, bool bOnlyIfOverridden = true);
 
 	static bool ShouldHandleAsImplementableEvent(UFunction* Function);
 
diff --git a/Engine/Source/Developer/BlueprintCompilerCppBackend/Private/BlueprintCompilerCppBackendValueHelper.cpp b/Engine/Source/Developer/BlueprintCompilerCppBackend/Private/BlueprintCompilerCppBackendValueHelper.cpp
--- a/Engine/Source/Developer/BlueprintCompilerCppBackend/Private/BlueprintCompilerCppBackendValueHelper.cpp
+++ b/Engine/Source/Developer/BlueprintCompilerCppBackend/Private/BlueprintCompilerCppBackendValueHelper.cpp
@@ -64,7 +64,6 @@ void FEmitDefaultValueHelper::OuterGenerate(FEmitterLocalContext& Context
 			}
 			else if (bNoexportProperty || Property->HasAnyPropertyFlags(CPF_NativeAccessSpecifierPrivate) || (!bAllowProtected && Property->HasAnyPropertyFlags(CPF_NativeAccessSpecifierProtected)))
 			{
-				ensure(EPropertyAccessOperator::None != AccessOperator);
 				const FString OperatorStr = (EPropertyAccessOperator::Dot == AccessOperator) ? TEXT("&") : TEXT("");
 				const FString ContainerStr = (EPropertyAccessOperator::None == AccessOperator) ? TEXT("this") : OuterPath;
 				const FString GetPtrStr = bNoexportProperty
@@ -269,7 +268,7 @@ bool FEmitDefaultValueHelper::SpecialStructureConstructor(const UScriptStruct* S
 		if (OutResult)
 		{
 			const FLinearColor* LinearColor = reinterpret_cast<const FLinearColor*>(ValuePtr);
-			*OutResult = FString::Printf(TEXT("FLinearColor(%f, %f, %f, %f)"), LinearColor->R, LinearColor->B, LinearColor->G, LinearColor->A);
+			*OutResult = FString::Printf(TEXT("FLinearColor(%f, %f, %f, %f)"), LinearColor->R, LinearColor->G, LinearColor->B, LinearColor->A);
 		}
 		return true;
 	}
@@ -279,7 +278,7 @@ bool FEmitDefaultValueHelper::SpecialStructureConstructor(const UScriptStruct* S
 		if (OutResult)
 		{
 			const FColor* Color = reinterpret_cast<const FColor*>(ValuePtr);
-			*OutResult = FString::Printf(TEXT("FColor(%d, %d, %d, %d)"), Color->R, Color->B, Color->G, Color->A);
+			*OutResult = FString::Printf(TEXT("FColor(%d, %d, %d, %d)"), Color->R, Color->G, Color->B, Color->A);
 		}
 		return true;
 	}
@@ -510,6 +509,17 @@ struct FNonativeComponentData
 		}
 
 	}
+
+	void EmitForcedPostLoad(FEmitterLocalContext& Context)
+	{
+		Context.AddLine(FString::Printf(TEXT("if(%s && !%s->IsTemplate())"), *NativeVariablePropertyName, *NativeVariablePropertyName));
+		Context.AddLine(TEXT("{"));
+		Context.IncreaseIndent();
+		Context.AddLine(FString::Printf(TEXT("%s->SetFlags(RF_NeedPostLoad |RF_NeedPostLoadSubobjects);"), *NativeVariablePropertyName));
+		Context.AddLine(FString::Printf(TEXT("%s->ConditionalPostLoad();"), *NativeVariablePropertyName));
+		Context.DecreaseIndent();
+		Context.AddLine(TEXT("}"));
+	}
 };
 
 FString FEmitDefaultValueHelper::HandleNonNativeComponent(FEmitterLocalContext& Context, const USCS_Node* Node, TSet<const UProperty*>& OutHandledProperties, TArray<FString>& NativeCreatedComponentProperties, const USCS_Node* ParentNode, TArray<FNonativeComponentData>& ComponenntsToInit)
@@ -853,7 +863,7 @@ void FEmitDefaultValueHelper::GenerateConstructor(FEmitterLocalContext& Context)
 
 	// Components that must be fixed after serialization
 	TArray<FString> NativeCreatedComponentProperties;
-
+	TArray<FNonativeComponentData> ComponentsToInit;
 	{
 		UObject* CDO = BPGC->GetDefaultObject(false);
 
@@ -923,8 +933,6 @@ void FEmitDefaultValueHelper::GenerateConstructor(FEmitterLocalContext& Context)
 		const bool bErrorFree = UBlueprintGeneratedClass::GetGeneratedClassesHierarchy(BPGC, BPGCStack);
 		if (bErrorFree)
 		{
-			TArray<FNonativeComponentData> ComponentsToInit;
-
 			// Start at the base of the hierarchy so that dependencies are handled first.
 			for (int32 i = BPGCStack.Num() - 1; i >= 0; --i)
 			{
@@ -962,10 +970,9 @@ void FEmitDefaultValueHelper::GenerateConstructor(FEmitterLocalContext& Context)
 		// Generate ctor init code for generated Blueprint class property values that may differ from parent class defaults (or that otherwise belong to the generated Blueprint class).
 		for (auto Property : TFieldRange<const UProperty>(BPGC))
 		{
-			const bool bNewProperty = Property->GetOwnerStruct() == BPGC;
-			const bool bIsAccessible = bNewProperty || !Property->HasAnyPropertyFlags(CPF_NativeAccessSpecifierPrivate);
-			if (bIsAccessible && !HandledProperties.Contains(Property))
+			if (!HandledProperties.Contains(Property))
 			{
+				const bool bNewProperty = Property->GetOwnerStruct() == BPGC;
 				OuterGenerate(Context, Property, TEXT(""), reinterpret_cast<const uint8*>(CDO), bNewProperty ? nullptr : reinterpret_cast<const uint8*>(ParentCDO), EPropertyAccessOperator::None, true);
 			}
 		}
@@ -973,7 +980,24 @@ void FEmitDefaultValueHelper::GenerateConstructor(FEmitterLocalContext& Context)
 	Context.DecreaseIndent();
 	Context.AddLine(TEXT("}"));
 
+	// TODO: this mechanism could be required by other instanced subobjects.
 	Context.CurrentCodeType = FEmitterLocalContext::EGeneratedCodeType::Regular;
+	Context.ResetPropertiesForInaccessibleStructs();
+	if (BPGC->IsChildOf<AActor>() && ComponentsToInit.Num())
+	{
+		Context.Header.AddLine(TEXT("virtual void FixComponentsFromDynamicClass() override;"));
+		Context.AddLine(FString::Printf(TEXT("void %s::FixComponentsFromDynamicClass()"), *CppClassName));
+		Context.AddLine(TEXT("{"));
+		Context.IncreaseIndent();
+		Context.AddLine(TEXT("Super::FixComponentsFromDynamicClass();"));
+		for (auto& ComponentToInit : ComponentsToInit)
+		{
+			ComponentToInit.EmitForcedPostLoad(Context);
+		}
+		Context.DecreaseIndent();
+		Context.AddLine(TEXT("}"));
+	}
+
 	Context.ResetPropertiesForInaccessibleStructs();
 	Context.AddLine(FString::Printf(TEXT("void %s::PostLoadSubobjects(FObjectInstancingGraph* OuterInstanceGraph)"), *CppClassName));
 	Context.AddLine(TEXT("{"));
@@ -981,7 +1005,7 @@ void FEmitDefaultValueHelper::GenerateConstructor(FEmitterLocalContext& Context)
 	Context.AddLine(TEXT("Super::PostLoadSubobjects(OuterInstanceGraph);"));
 	for (auto& ComponentToFix : NativeCreatedComponentProperties)
 	{
-		Context.AddLine(FString::Printf(TEXT("if(ensure(%s))"), *ComponentToFix));
+		Context.AddLine(FString::Printf(TEXT("if(%s)"), *ComponentToFix));
 		Context.AddLine(TEXT("{"));
 		Context.IncreaseIndent();
 		Context.AddLine(FString::Printf(TEXT("%s->CreationMethod = EComponentCreationMethod::Native;"), *ComponentToFix));
diff --git a/Engine/Source/Developer/BlueprintProfiler/Private/BlueprintProfiler.cpp b/Engine/Source/Developer/BlueprintProfiler/Private/BlueprintProfiler.cpp
--- a/Engine/Source/Developer/BlueprintProfiler/Private/BlueprintProfiler.cpp
+++ b/Engine/Source/Developer/BlueprintProfiler/Private/BlueprintProfiler.cpp
@@ -185,9 +185,12 @@ void FBlueprintProfiler::ProcessEventProfilingData()
 					{
 						FEventRange NewEventRange;
 						NewEventRange.BlueprintContext = GetBlueprintContext(CurrEvent.GetObjectPath());
-						NewEventRange.InstanceName = MapBlueprintInstance(NewEventRange.BlueprintContext, InstanceEvent.GetObjectPath());
-						NewEventRange.StartIdx = EventIdx;
-						ScriptEventRanges.Push(NewEventRange);
+						if (NewEventRange.BlueprintContext.IsValid())
+						{
+							NewEventRange.InstanceName = MapBlueprintInstance(NewEventRange.BlueprintContext, InstanceEvent.GetObjectPath());
+							NewEventRange.StartIdx = EventIdx;
+							ScriptEventRanges.Push(NewEventRange);
+						}
 					}
 				}
 				InstrumentationEventQueue.RemoveAt(EventIdx, 2, false);
@@ -196,13 +199,15 @@ void FBlueprintProfiler::ProcessEventProfilingData()
 			case EScriptInstrumentation::Event:
 			{
 				// Nested events such as calls from event dispatchers.
-				check(ScriptEventRanges.Num() > 0);
-				FEventRange& LastEventRange = ScriptEventRanges.Last();
-				FEventRange NewEventRange;
-				NewEventRange.BlueprintContext = LastEventRange.BlueprintContext;
-				NewEventRange.InstanceName = LastEventRange.InstanceName;
-				NewEventRange.StartIdx = EventIdx;
-				ScriptEventRanges.Push(NewEventRange);
+				if (ScriptEventRanges.Num() > 0)
+				{
+					FEventRange& LastEventRange = ScriptEventRanges.Last();
+					FEventRange NewEventRange;
+					NewEventRange.BlueprintContext = LastEventRange.BlueprintContext;
+					NewEventRange.InstanceName = LastEventRange.InstanceName;
+					NewEventRange.StartIdx = EventIdx;
+					ScriptEventRanges.Push(NewEventRange);
+				}
 				break;
 			}
 			case EScriptInstrumentation::Stop:
@@ -346,7 +351,6 @@ FName FBlueprintProfiler::MapBlueprintInstance(TSharedPtr<FBlueprintExecutionCon
 		// Link to parent blueprint entry
 		TSharedPtr<FScriptExecutionBlueprint> BlueprintNode = BlueprintContext->GetBlueprintExecNode();
 		BlueprintNode->AddInstance(InstanceNode);
-		InstanceNode->SetParentNode(BlueprintNode);
 		// Fill out events from the blueprint root node
 		for (int32 NodeIdx = 0; NodeIdx < BlueprintNode->GetNumChildren(); ++NodeIdx)
 		{
diff --git a/Engine/Source/Developer/BlueprintProfiler/Private/ScriptInstrumentationPlayback.cpp b/Engine/Source/Developer/BlueprintProfiler/Private/ScriptInstrumentationPlayback.cpp
--- a/Engine/Source/Developer/BlueprintProfiler/Private/ScriptInstrumentationPlayback.cpp
+++ b/Engine/Source/Developer/BlueprintProfiler/Private/ScriptInstrumentationPlayback.cpp
@@ -47,6 +47,12 @@ bool FBlueprintExecutionContext::InitialiseContext(const FString& BlueprintPath)
 	return false;
 }
 
+void FBlueprintExecutionContext::AddEventNode(TSharedPtr<FScriptExecutionNode> EventExecNode)
+{
+	check(BlueprintNode.IsValid());
+	BlueprintNode->AddChildNode(EventExecNode);
+}
+
 bool FBlueprintExecutionContext::HasProfilerDataForInstance(const FName InstanceName) const
 {
 	bool bHasInstanceData = false;
@@ -138,16 +144,15 @@ TSharedPtr<FBlueprintFunctionContext> FBlueprintExecutionContext::GetFunctionCon
 	return Result;
 }
 
-TSharedPtr<FBlueprintFunctionContext> FBlueprintExecutionContext::CreateFunctionContext(const FName FunctionNameIn)
+TSharedPtr<FBlueprintFunctionContext> FBlueprintExecutionContext::CreateFunctionContext(UEdGraph* Graph)
 {
-	FName FunctionName = (FunctionNameIn ==  UEdGraphSchema_K2::GN_EventGraph) ? UbergraphFunctionName : FunctionNameIn;
+	FName FunctionName = Graph->GetFName();
+	FunctionName = (FunctionName == UEdGraphSchema_K2::GN_EventGraph) ? UbergraphFunctionName : FunctionName;
 	TSharedPtr<FBlueprintFunctionContext>& Result = FunctionContexts.FindOrAdd(FunctionName);
 	if (!Result.IsValid())
 	{
-		if (UFunction* Function = BlueprintClass->FindFunctionByName(FunctionName))
-		{
-			Result = MakeShareable(new FBlueprintFunctionContext(Function, BlueprintClass));
-		}
+		Result = MakeShareable(new FBlueprintFunctionContext);
+		Result->InitialiseContextFromGraph(AsShared(), Graph);
 	}
 	check(Result.IsValid());
 	return Result;
@@ -197,89 +202,203 @@ void FBlueprintExecutionContext::UpdateConnectedStats()
 bool FBlueprintExecutionContext::MapBlueprintExecution()
 {
 	bool bMappingSuccessful = false;
+	UBlueprint* BlueprintToMap = Blueprint.Get();
+	UBlueprintGeneratedClass* BPGC = BlueprintClass.Get();
 
-	if (UBlueprint* BlueprintToMap = Blueprint.Get())
+	if (BlueprintToMap && BPGC)
 	{
-		// Find all event nodes for event graph and UCS
+		// Map all blueprint graphs and event entry points
 		TArray<UEdGraph*> Graphs;
-		TArray<UK2Node_Event*> EventNodes;
-		TArray<UK2Node_FunctionEntry*> FunctionEntryNodes;
-		TArray<UK2Node*> InputEventNodes;
+		TArray<UEdGraph*> FunctionGraphs;
+		TSet<UEdGraph*> Macros;
 		BlueprintToMap->GetAllGraphs(Graphs);
-		for (auto GraphIter : Graphs)
+		for (auto Graph : Graphs)
+		{
+			FName FunctionName = Graph->GetFName();
+			FunctionName = FunctionName == UEdGraphSchema_K2::GN_EventGraph ? UbergraphFunctionName : FunctionName;
+			if (UFunction* ScriptFunction = BPGC->FindFunctionByName(FunctionName))
+			{
+				FunctionGraphs.Add(Graph);
+			}
+		}
+		// Create function stubs
+		for (auto Graph : FunctionGraphs)
+		{
+			CreateFunctionContext(Graph);
+		}
+		// Map the function context nodes
+		for (auto Context : FunctionContexts)
+		{
+			Context.Value->MapFunction();
+		}
+		bMappingSuccessful = true;
+	}
+	return bMappingSuccessful;
+}
+
+//////////////////////////////////////////////////////////////////////
+// FBlueprintFunctionContext
+
+void FBlueprintFunctionContext::InitialiseContextFromGraph(TSharedPtr<FBlueprintExecutionContext> BlueprintContextIn, UEdGraph* Graph)
+{
+	if (Graph)
+	{
+		// Initialise Blueprint references
+		BlueprintContext = BlueprintContextIn;
+		UBlueprint* Blueprint = BlueprintContextIn->GetBlueprint().Get();
+		UBlueprintGeneratedClass* BPClass = BlueprintContextIn->GetBlueprintClass().Get();
+		if (Blueprint && BPClass)
 		{
-			if (GraphIter->GetFName() == UEdGraphSchema_K2::GN_EventGraph)
+			// Instantiate context
+			BlueprintClass = BPClass;
+			const bool bEventGraph = Graph->GetFName() == UEdGraphSchema_K2::GN_EventGraph;
+			FunctionName = bEventGraph && BPClass->UberGraphFunction ? BPClass->UberGraphFunction->GetFName() : Graph->GetFName();
+			Function = BPClass->FindFunctionByName(FunctionName);
+			Function = Function.IsValid() ? Function : BPClass->UberGraphFunction;
+			if (bEventGraph)
 			{
-				GraphIter->GetNodesOfClass<UK2Node_Event>(EventNodes);
+				// Map events
+				TArray<UK2Node_Event*> GraphEventNodes;
+				Graph->GetNodesOfClass<UK2Node_Event>(GraphEventNodes);
+
+				for (auto EventNode : GraphEventNodes)
+				{
+					FName EventName = EventNode->GetFunctionName();
+					FScriptExecNodeParams EventParams;
+					EventParams.NodeName = EventName;
+					EventParams.ObservedObject = EventNode;
+					EventParams.DisplayName = EventNode->GetNodeTitle(ENodeTitleType::ListView);
+					EventParams.Tooltip = LOCTEXT("NavigateToEventLocationHyperlink_ToolTip", "Navigate to the Event");
+					EventParams.NodeFlags = EScriptExecutionNodeFlags::Event;
+					EventParams.IconColor = FLinearColor(0.91f, 0.16f, 0.16f);
+					const FSlateBrush* EventIcon = EventNode->ShowPaletteIconOnNode() ?	FEditorStyle::GetBrush(EventNode->GetPaletteIcon(EventParams.IconColor)) :
+																						FEditorStyle::GetBrush(TEXT("BlueprintProfiler.BPNode"));
+					EventParams.Icon = const_cast<FSlateBrush*>(EventIcon);
+					TSharedPtr<FScriptExecutionNode> EventExecNode = CreateExecutionNode(EventParams);
+					AddEntryPoint(EventExecNode);
+					BlueprintContextIn->AddEventNode(EventExecNode);
+				}
+				// Create Input events
+				TArray<UK2Node*> InputEventNodes;
+				Graph->GetNodesOfClassEx<UK2Node_InputAction, UK2Node>(InputEventNodes);
+				Graph->GetNodesOfClassEx<UK2Node_InputKey, UK2Node>(InputEventNodes);
+				Graph->GetNodesOfClassEx<UK2Node_InputTouch, UK2Node>(InputEventNodes);
+				if (InputEventNodes.Num())
+				{
+					CreateInputEvents(BlueprintContextIn, InputEventNodes);
+				}
 			}
-			if (GraphIter->GetFName() == UEdGraphSchema_K2::FN_UserConstructionScript)
+			else
 			{
-				GraphIter->GetNodesOfClass<UK2Node_FunctionEntry>(FunctionEntryNodes);
+				// Map execution paths for each function entry nodes
+				TArray<UK2Node_FunctionEntry*> FunctionEntryNodes;
+				Graph->GetNodesOfClass<UK2Node_FunctionEntry>(FunctionEntryNodes);
+				const bool bConstructionScript = FunctionName == UEdGraphSchema_K2::FN_UserConstructionScript;
+				// Create function node
+				for (auto FunctionEntry : FunctionEntryNodes)
+				{
+					FScriptExecNodeParams FunctionNodeParams;
+					FunctionNodeParams.NodeName = FunctionName;
+					FunctionNodeParams.ObservedObject = FunctionEntry;
+					FunctionNodeParams.DisplayName = FText::FromName(FunctionName);
+					FunctionNodeParams.Tooltip = LOCTEXT("NavigateToFunctionLocationHyperlink_ToolTip", "Navigate to the Function");
+					FunctionNodeParams.NodeFlags = EScriptExecutionNodeFlags::Event;
+					FunctionNodeParams.IconColor = FLinearColor(0.91f, 0.16f, 0.16f);
+					const FSlateBrush* Icon = FunctionEntry->ShowPaletteIconOnNode() ?	FEditorStyle::GetBrush(FunctionEntry->GetPaletteIcon(FunctionNodeParams.IconColor)) :
+																						FEditorStyle::GetBrush(TEXT("BlueprintProfiler.BPNode"));
+					FunctionNodeParams.Icon = const_cast<FSlateBrush*>(Icon);
+					TSharedPtr<FScriptExecutionNode> FunctionEntryNode = CreateExecutionNode(FunctionNodeParams);
+					AddEntryPoint(FunctionEntryNode);
+					if (bConstructionScript)
+					{
+						BlueprintContextIn->AddEventNode(FunctionEntryNode);
+					}
+				}
+			}
+			// Map all instanced sub graphs
+			TArray<UK2Node_Tunnel*> GraphTunnels;
+			Graph->GetNodesOfClass<UK2Node_Tunnel>(GraphTunnels);
+			// Map sub graphs / composites and macros
+			for (auto Tunnel : GraphTunnels)
+			{
+				MapTunnelInstance(Tunnel);
 			}
-			GraphIter->GetNodesOfClassEx<UK2Node_InputAction, UK2Node>(InputEventNodes);
-			GraphIter->GetNodesOfClassEx<UK2Node_InputKey, UK2Node>(InputEventNodes);
-			GraphIter->GetNodesOfClassEx<UK2Node_InputTouch, UK2Node>(InputEventNodes);
 		}
-		// Map execution paths for each event node
-		for (auto EventIter : EventNodes)
+	}
+}
+
+void FBlueprintFunctionContext::MapFunction()
+{
+	// Map all entry point execution paths
+	for (auto EntryPoint : EntryPoints)
+	{
+		if (const UEdGraphNode* EntryPointNode = Cast<UEdGraphNode>(EntryPoint->GetObservedObject()))
 		{
-			FName EventName = EventIter->GetFunctionName();
-			TSharedPtr<FBlueprintFunctionContext> EventContext = CreateFunctionContext(UbergraphFunctionName);
-			FScriptExecNodeParams EventParams;
-			EventParams.NodeName = EventName;
-			EventParams.ObservedObject = EventIter;
-			EventParams.OwningGraphName = NAME_None;
-			EventParams.DisplayName = EventIter->GetNodeTitle(ENodeTitleType::ListView);
-			EventParams.Tooltip = LOCTEXT("NavigateToEventLocationHyperlink_ToolTip", "Navigate to the Event");
-			EventParams.NodeFlags = EScriptExecutionNodeFlags::Event;
-			EventParams.IconColor = FLinearColor(0.91f, 0.16f, 0.16f);
-			const FSlateBrush* EventIcon = EventIter->ShowPaletteIconOnNode() ?	FEditorStyle::GetBrush(EventIter->GetPaletteIcon(EventParams.IconColor)) :
-																				FEditorStyle::GetBrush(TEXT("BlueprintProfiler.BPNode"));
-			EventParams.Icon = const_cast<FSlateBrush*>(EventIcon);
-			TSharedPtr<FScriptExecutionNode> EventExecNode = EventContext->CreateExecutionNode(EventParams);
-			// Add event to mapped blueprint
-			BlueprintNode->AddChildNode(EventExecNode);
-			// Map the event execution flow
-			TSharedPtr<FScriptExecutionNode> MappedEvent = MapNodeExecution(EventContext, EventIter);
-			if (MappedEvent.IsValid())
+			TSharedPtr<FScriptExecutionNode> ExecutionNode = MapNodeExecution(const_cast<UEdGraphNode*>(EntryPointNode));
+			if (ExecutionNode.IsValid())
 			{
-				EventExecNode->AddChildNode(MappedEvent);
+				EntryPoint->AddChildNode(ExecutionNode);
 			}
 		}
-		// Create Input events
-		CreateInputEvents(InputEventNodes);
-		// Map execution paths for each function entry node
-		for (auto FunctionIter : FunctionEntryNodes)
+		// Check for Cyclic Linkage
+		TSet<TSharedPtr<FScriptExecutionNode>> Filter;
+		DetectCyclicLinks(EntryPoint, Filter);
+	}
+}
+
+bool FBlueprintFunctionContext::DetectCyclicLinks(TSharedPtr<FScriptExecutionNode> ExecNode, TSet<TSharedPtr<FScriptExecutionNode>>& Filter)
+{
+	if (Filter.Contains(ExecNode))
+	{
+		return true;
+	}
+	else
+	{
+		Filter.Add(ExecNode);
+		for (auto Child : ExecNode->GetChildNodes())
+		{
+			DetectCyclicLinks(Child, Filter);
+		}
+		for (auto LinkedNode : ExecNode->GetLinkedNodes())
 		{
-			FName FunctionName = FunctionIter->SignatureName;
-			TSharedPtr<FBlueprintFunctionContext> FunctionContext = CreateFunctionContext(FunctionName);
-			FScriptExecNodeParams FunctionParams;
-			FunctionParams.NodeName = FunctionName;
-			FunctionParams.ObservedObject = FunctionIter;
-			FunctionParams.OwningGraphName = NAME_None;
-			FunctionParams.DisplayName = FunctionIter->GetNodeTitle(ENodeTitleType::ListView);
-			FunctionParams.Tooltip = LOCTEXT("NavigateToFunctionLocationHyperlink_ToolTip", "Navigate to the Function");
-			FunctionParams.NodeFlags = EScriptExecutionNodeFlags::Event;
-			FunctionParams.IconColor = FLinearColor(0.91f, 0.16f, 0.16f);
-			const FSlateBrush* FunctionIcon = FunctionIter->ShowPaletteIconOnNode() ?	FEditorStyle::GetBrush(FunctionIter->GetPaletteIcon(FunctionParams.IconColor)) :
-																						FEditorStyle::GetBrush(TEXT("BlueprintProfiler.BPNode"));
-			FunctionParams.Icon = const_cast<FSlateBrush*>(FunctionIcon);
-			TSharedPtr<FScriptExecutionNode> FunctionExecNode = FunctionContext->CreateExecutionNode(FunctionParams);
-			// Add Function to mapped blueprint
-			BlueprintNode->AddChildNode(FunctionExecNode);
-			// Map the Function execution flow
-			TSharedPtr<FScriptExecutionNode> MappedFunction = MapNodeExecution(FunctionContext, FunctionIter);
-			if (MappedFunction.IsValid())
+			if (!LinkedNode.Value->HasFlags(EScriptExecutionNodeFlags::PureStats))
 			{
-				FunctionExecNode->AddChildNode(MappedFunction);
+				TSharedPtr<FScriptExecutionNode> LinkedExecNode = LinkedNode.Value;
+				if (LinkedExecNode->HasFlags(EScriptExecutionNodeFlags::FunctionTunnel))
+				{
+					TSharedPtr<FScriptExecutionNode> TunnelBoundry = LinkedExecNode->GetLinkedNodeByScriptOffset(LinkedNode.Key);
+					if (TunnelBoundry.IsValid())
+					{
+						TSharedPtr<FScriptExecutionNode> TunnelExec = TunnelBoundry->GetLinkedNodeByScriptOffset(LinkedNode.Key);
+						if (TunnelExec.IsValid())
+						{
+							ExecNode = TunnelExec;
+						}
+					}
+				}
+				if (DetectCyclicLinks(LinkedExecNode, Filter))
+				{
+					// Break links and flag cycle linkage.
+					FScriptExecNodeParams CycleLinkParams;
+					const FString NodeName = FString::Printf(TEXT("CyclicLinkTo_%i_%s"), LinkedNode.Key, *LinkedExecNode->GetName().ToString());
+					CycleLinkParams.NodeName = FName(*NodeName);
+					CycleLinkParams.ObservedObject = LinkedExecNode->GetObservedObject();
+					CycleLinkParams.DisplayName = LinkedExecNode->GetDisplayName();
+					CycleLinkParams.Tooltip = LOCTEXT("CyclicLink_ToolTip", "Cyclic Link");
+					CycleLinkParams.NodeFlags = EScriptExecutionNodeFlags::CyclicLinkage|EScriptExecutionNodeFlags::ExecPin|EScriptExecutionNodeFlags::InvalidTrace;
+					CycleLinkParams.IconColor = FLinearColor(0.91f, 0.16f, 0.16f);
+					const FSlateBrush* LinkIcon = LinkedExecNode->GetIcon();
+					CycleLinkParams.Icon = const_cast<FSlateBrush*>(LinkIcon);
+					TSharedPtr<FScriptExecutionNode> NewLink = CreateExecutionNode(CycleLinkParams);
+					ExecNode->AddLinkedNode(LinkedNode.Key, NewLink);
+				}
 			}
 		}
-		bMappingSuccessful = true;
 	}
-	return bMappingSuccessful;
+	return false;
 }
 
-void FBlueprintExecutionContext::CreateInputEvents(const TArray<UK2Node*>& InputEventNodes)
+void FBlueprintFunctionContext::CreateInputEvents(TSharedPtr<FBlueprintExecutionContext> BlueprintContextIn, const TArray<UK2Node*>& InputEventNodes)
 {
 	struct FInputEventDesc
 	{
@@ -375,7 +494,6 @@ void FBlueprintExecutionContext::CreateInputEvents(const TArray<UK2Node*>& Input
 	{
 		UEdGraphNode* GraphNode = InputEventDesc.GraphNode;
 		FName EventName = InputEventDesc.EventFunction->GetFName();
-		TSharedPtr<FBlueprintFunctionContext> EventContext = CreateFunctionContext(UbergraphFunctionName);
 		FScriptExecNodeParams EventParams;
 		EventParams.NodeName = EventName;
 		EventParams.ObservedObject = GraphNode;
@@ -387,24 +505,18 @@ void FBlueprintExecutionContext::CreateInputEvents(const TArray<UK2Node*>& Input
 		const FSlateBrush* EventIcon = GraphNode->ShowPaletteIconOnNode() ?	FEditorStyle::GetBrush(GraphNode->GetPaletteIcon(EventParams.IconColor)) :
 																			FEditorStyle::GetBrush(TEXT("BlueprintProfiler.BPNode"));
 		EventParams.Icon = const_cast<FSlateBrush*>(EventIcon);
-		TSharedPtr<FScriptExecutionNode> EventExecNode = EventContext->CreateExecutionNode(EventParams);
-		// Add Function to mapped blueprint
-		BlueprintNode->AddChildNode(EventExecNode);
-		// Map the Function execution flow
-		TSharedPtr<FScriptExecutionNode> MappedFunction = MapNodeExecution(EventContext, GraphNode);
-		if (MappedFunction.IsValid())
-		{
-			EventExecNode->AddChildNode(MappedFunction);
-		}
+		TSharedPtr<FScriptExecutionNode> EventExecNode = CreateExecutionNode(EventParams);
+		AddEntryPoint(EventExecNode);
+		BlueprintContextIn->AddEventNode(EventExecNode);
 	}
 }
 
-TSharedPtr<FScriptExecutionNode> FBlueprintExecutionContext::MapNodeExecution(TSharedPtr<FBlueprintFunctionContext> FunctionContext, UEdGraphNode* NodeToMap)
+TSharedPtr<FScriptExecutionNode> FBlueprintFunctionContext::MapNodeExecution(UEdGraphNode* NodeToMap)
 {
 	TSharedPtr<FScriptExecutionNode> MappedNode;
 	if (NodeToMap)
 	{
-		MappedNode = FunctionContext->GetProfilerDataForNode(NodeToMap->GetFName());
+		MappedNode = GetProfilerDataForNode(NodeToMap->GetFName());
 		if (!MappedNode.IsValid())
 		{
 			FScriptExecNodeParams NodeParams;
@@ -417,7 +529,7 @@ TSharedPtr<FScriptExecutionNode> FBlueprintExecutionContext::MapNodeExecution(TS
 			const FSlateBrush* NodeIcon = NodeToMap->ShowPaletteIconOnNode() ?	FEditorStyle::GetBrush(NodeToMap->GetPaletteIcon(NodeParams.IconColor)) :
 																				FEditorStyle::GetBrush(TEXT("BlueprintProfiler.BPNode"));
 			NodeParams.Icon = const_cast<FSlateBrush*>(NodeIcon);
-			MappedNode = FunctionContext->CreateExecutionNode(NodeParams);
+			MappedNode = CreateExecutionNode(NodeParams);
 			// Evaluate Execution and Input Pins
 			TArray<UEdGraphPin*> ExecPins;
 			TArray<UEdGraphPin*> InputPins;
@@ -432,152 +544,278 @@ TSharedPtr<FScriptExecutionNode> FBlueprintExecutionContext::MapNodeExecution(TS
 					InputPins.Add(Pin);
 				}
 			}
-
+			// Flag as pure node or build pure node script range.
 			if (ExecPins.Num() == 0)
 			{
 				MappedNode->AddFlags(EScriptExecutionNodeFlags::PureNode);
 			}
 			else
 			{
-				MappedNode->SetPureNodeScriptCodeRange(FunctionContext->GetPureNodeScriptCodeRange(NodeToMap));
+				MappedNode->SetPureNodeScriptCodeRange(GetPureNodeScriptCodeRange(NodeToMap));
 			}
-
-			const bool bBranchNode = ExecPins.Num() > 1;
-			if (bBranchNode)
+			// Determine a few node characteristics to guide processing.
+			if (ExecPins.Num() > 1)
 			{
-				// Determine branch exec type
-				if (NodeToMap->IsA<UK2Node_ExecutionSequence>())
-				{
-					MappedNode->AddFlags(EScriptExecutionNodeFlags::SequentialBranch);
-				}
-				else
-				{
-					MappedNode->AddFlags(EScriptExecutionNodeFlags::ConditionalBranch);
-				}
+				const bool bExecSequence = NodeToMap->IsA<UK2Node_ExecutionSequence>();
+				MappedNode->AddFlags(bExecSequence ? EScriptExecutionNodeFlags::SequentialBranch : EScriptExecutionNodeFlags::ConditionalBranch);
 			}
 			// Evaluate non-exec input pins (pure node execution chains)
-			TArray<int32> PinScriptCodeOffsets;
-			for (auto InputPin : InputPins)
+			if (InputPins.Num())
+			{
+				MapInputPins(MappedNode, InputPins);
+			}
+			// Evaluate exec output pins (execution chains)
+			if (ExecPins.Num())
+			{
+				MapExecPins(MappedNode, ExecPins);
+			}
+			// Evaluate Children for call sites
+			if (UK2Node_CallFunction* FunctionCallSite = Cast<UK2Node_CallFunction>(NodeToMap))
 			{
-				// If this input pin is linked to a pure node in the source graph, create and map all known execution paths for it.
-				for (auto LinkedPin : InputPin->LinkedTo)
+				const UEdGraphNode* FunctionNode = nullptr; 
+				if (UEdGraph* CalledGraph = FunctionCallSite->GetFunctionGraph(FunctionNode))
 				{
-					// Note: Intermediate pure nodes can have output pins that masquerade as impure node output pins when links are "moved" from the source graph (thus
-					// resulting in a false association here with one or more script code offsets), so we must first ensure that the link is really to a pure node output.
-					UK2Node* OwningNode = Cast<UK2Node>(LinkedPin->GetOwningNode());
-					if (OwningNode && OwningNode->IsNodePure())
+					// Update Exec node
+					const bool bEventCall = FunctionNode ? FunctionNode->IsA<UK2Node_Event>() : false;
+					MappedNode->AddFlags(EScriptExecutionNodeFlags::FunctionCall);
+					MappedNode->SetIconColor(FLinearColor(0.46f, 0.54f, 0.95f));
+					MappedNode->SetToolTipText(LOCTEXT("NavigateToFunctionLocationHyperlink_ToolTip", "Navigate to the Function Callsite"));
+					// Don't add entry points for events.
+					if (!bEventCall)
 					{
-						FunctionContext->GetAllCodeLocationsFromPin(LinkedPin, PinScriptCodeOffsets);
-						if (PinScriptCodeOffsets.Num() > 0)
+						// Update the function context
+						TSharedPtr<FBlueprintFunctionContext> NewFunctionContext = BlueprintContext.Pin()->GetFunctionContext(CalledGraph->GetFName());
+						if (NewFunctionContext.IsValid())
 						{
-							TSharedPtr<FScriptExecutionNode> PureNode = MapNodeExecution(FunctionContext, OwningNode);
-							check(PureNode.IsValid() && PureNode->IsPureNode());
-							for (int32 i = 0; i < PinScriptCodeOffsets.Num(); ++i)
-							{
-								MappedNode->AddPureNode(PinScriptCodeOffsets[i], PureNode);
-							}
+							NewFunctionContext->AddCallSiteEntryPointsToNode(MappedNode);
 						}
 					}
 				}
 			}
-			for (auto Pin : ExecPins)
+			else if (UK2Node_CustomEvent* CustomEvent = Cast<UK2Node_CustomEvent>(NodeToMap))
 			{
-				const UEdGraphPin* LinkedPin = Pin->LinkedTo.Num() > 0 ? Pin->LinkedTo[0] : nullptr;
+				MappedNode->AddFlags(EScriptExecutionNodeFlags::CustomEvent);
+			}
+		}
+	}
+	return MappedNode;
+}
+
+void FBlueprintFunctionContext::MapInputPins(TSharedPtr<FScriptExecutionNode> ExecNode, const TArray<UEdGraphPin*>& Pins)
+{
+	TArray<int32> PinScriptCodeOffsets;
+	TSharedPtr<FScriptExecutionNode> PureChainRootNode = ExecNode;
 
-				if (bBranchNode)
+	for (auto InputPin : Pins)
+	{
+		// If this input pin is linked to a pure node in the source graph, create and map all known execution paths for it.
+		for (auto LinkedPin : InputPin->LinkedTo)
+		{
+			// Note: Intermediate pure nodes can have output pins that masquerade as impure node output pins when links are "moved" from the source graph (thus
+			// resulting in a false association here with one or more script code offsets), so we must first ensure that the link is really to a pure node output.
+			UK2Node* OwningNode = Cast<UK2Node>(LinkedPin->GetOwningNode());
+			if (OwningNode && OwningNode->IsNodePure())
+			{
+				GetAllCodeLocationsFromPin(LinkedPin, PinScriptCodeOffsets);
+				if (PinScriptCodeOffsets.Num() > 0)
 				{
-					const int32 PinScriptCodeOffset = FunctionContext->GetCodeLocationFromPin(Pin);
-					// Check we have a valid script offset or that the pin isn't linked.
-					check((PinScriptCodeOffset != INDEX_NONE)||(LinkedPin == nullptr));
-					FScriptExecNodeParams LinkNodeParams;
-					LinkNodeParams.NodeName = Pin->GetFName();
-					LinkNodeParams.ObservedObject = Pin;
-					LinkNodeParams.DisplayName = Pin->GetDisplayName();
-					LinkNodeParams.Tooltip = LOCTEXT("ExecPin_ToolTip", "Expand execution path");
-					LinkNodeParams.NodeFlags = EScriptExecutionNodeFlags::ExecPin;
-					LinkNodeParams.IconColor = FLinearColor(1.f, 1.f, 1.f, 0.8f);
-					const FSlateBrush* Icon = LinkedPin ?	FEditorStyle::GetBrush(TEXT("BlueprintProfiler.BPPinConnected")) : 
-															FEditorStyle::GetBrush(TEXT("BlueprintProfiler.BPPinDisconnected"));
-					LinkNodeParams.Icon = const_cast<FSlateBrush*>(Icon);
-					TSharedPtr<FScriptExecutionNode> LinkedNode = FunctionContext->CreateExecutionNode(LinkNodeParams);
-					MappedNode->AddLinkedNode(PinScriptCodeOffset, LinkedNode);
-					// Map pin exec flow
-					if (LinkedPin)
+					// Add a pure chain container node as the root, if it's not already in place.
+					if (PureChainRootNode == ExecNode && !ExecNode->IsPureNode())
 					{
-						TSharedPtr<FScriptExecutionNode> ExecPinNode = MapNodeExecution(FunctionContext, LinkedPin->GetOwningNode());
-						if (ExecPinNode.IsValid())
-						{
-							LinkedNode->AddChildNode(ExecPinNode);
-						}
+						FScriptExecNodeParams PureChainParams;
+						static const FString PureChainNodeNameSuffix = TEXT("__PROFILER_PureTime");
+						FString PureChainNodeNameString = ExecNode->GetName().ToString() + PureChainNodeNameSuffix;
+						PureChainParams.NodeName = FName(*PureChainNodeNameString);
+						PureChainParams.DisplayName = LOCTEXT("PureChain_DisplayName", "Pure Time");
+						PureChainParams.Tooltip = LOCTEXT("PureChain_ToolTip", "Expand pure node timing");
+						PureChainParams.NodeFlags = EScriptExecutionNodeFlags::PureChain;
+						PureChainParams.IconColor = FLinearColor(1.f, 1.f, 1.f, 0.8f);
+						const FSlateBrush* Icon = FEditorStyle::GetBrush(TEXT("BlueprintProfiler.BPNode"));
+						PureChainParams.Icon = const_cast<FSlateBrush*>(Icon);
+						PureChainRootNode = CreateExecutionNode(PureChainParams);
+						PureChainRootNode->SetPureNodeScriptCodeRange(ExecNode->GetPureNodeScriptCodeRange());
+						ExecNode->AddChildNode(PureChainRootNode);
 					}
+
+					TSharedPtr<FScriptExecutionNode> PureNode = MapNodeExecution(OwningNode);
+					check(PureNode.IsValid() && PureNode->IsPureNode());
+					for (int32 i = 0; i < PinScriptCodeOffsets.Num(); ++i)
+					{
+						PureChainRootNode->AddLinkedNode(PinScriptCodeOffsets[i], PureNode);
+					}
+				}
+			}
+		}
+	}
+}
+
+void FBlueprintFunctionContext::MapExecPins(TSharedPtr<FScriptExecutionNode> ExecNode, const TArray<UEdGraphPin*>& Pins)
+{
+	const bool bBranchedExecution = ExecNode->IsBranch();
+	int32 NumUnwiredPins = 0;
+	for (auto Pin : Pins)
+	{
+		int32 PinScriptCodeOffset = GetCodeLocationFromPin(Pin);
+		const bool bInvalidTrace = PinScriptCodeOffset == INDEX_NONE;
+		PinScriptCodeOffset = bInvalidTrace ? NumUnwiredPins++ : PinScriptCodeOffset;
+		TSharedPtr<FScriptExecutionNode> PinExecNode;
+		const UEdGraphPin* LinkedPin = Pin->LinkedTo.Num() > 0 ? Pin->LinkedTo[0] : nullptr;
+		UEdGraphNode* LinkedPinNode = LinkedPin ? LinkedPin->GetOwningNode() : nullptr;
+		const bool bTunnelBoundry = LinkedPinNode ? LinkedPinNode->IsA<UK2Node_Tunnel>() : false;
+		// Create any neccesary dummy pins for branches
+		if (bBranchedExecution)
+		{
+			FScriptExecNodeParams LinkNodeParams;
+			LinkNodeParams.NodeName = Pin->GetFName();
+			LinkNodeParams.ObservedObject = Pin;
+			LinkNodeParams.DisplayName = Pin->GetDisplayName();
+			LinkNodeParams.Tooltip = LOCTEXT("ExecPin_ToolTip", "Expand execution path");
+			LinkNodeParams.NodeFlags = !bInvalidTrace ? EScriptExecutionNodeFlags::ExecPin : (EScriptExecutionNodeFlags::ExecPin|EScriptExecutionNodeFlags::InvalidTrace);
+			LinkNodeParams.IconColor = FLinearColor(1.f, 1.f, 1.f, 0.8f);
+			const FSlateBrush* Icon = LinkedPin ?	FEditorStyle::GetBrush(TEXT("BlueprintProfiler.BPPinConnected")) : 
+													FEditorStyle::GetBrush(TEXT("BlueprintProfiler.BPPinDisconnected"));
+			LinkNodeParams.Icon = const_cast<FSlateBrush*>(Icon);
+			PinExecNode = CreateExecutionNode(LinkNodeParams);
+			ExecNode->AddLinkedNode(PinScriptCodeOffset, PinExecNode);
+			if (!bTunnelBoundry)
+			{
+				TSharedPtr<FScriptExecutionNode> LinkedPinExecNode = MapNodeExecution(LinkedPinNode);
+				if (LinkedPinExecNode.IsValid())
+				{
+					PinExecNode->AddChildNode(LinkedPinExecNode);
 				}
-				else if (LinkedPin)
+			}
+		}
+		else
+		{
+			PinExecNode = ExecNode;
+			if (!bTunnelBoundry)
+			{
+				TSharedPtr<FScriptExecutionNode> LinkedPinExecNode = MapNodeExecution(LinkedPinNode);
+				if (LinkedPinExecNode.IsValid())
 				{
-					TSharedPtr<FScriptExecutionNode> ExecLinkedNode = MapNodeExecution(FunctionContext, LinkedPin->GetOwningNode());
-					if (ExecLinkedNode.IsValid())
+					PinExecNode->AddLinkedNode(PinScriptCodeOffset, LinkedPinExecNode);
+				}
+			}
+		}
+		// Handle tunnel boundries
+		if (bTunnelBoundry)
+		{
+			if (StagingTunnelName != NAME_None)
+			{
+				const FName TunnelExitName(*LinkedPin->PinName);
+				TSharedPtr<FScriptExecutionNode> TunnelExitNode = GetProfilerDataForNode(TunnelExitName);
+				// Check for re-entrant exit site
+				if (ExecNode->HasFlags(EScriptExecutionNodeFlags::SequentialBranch))
+				{
+					// Last pin in sequence is not re-entrant.
+					if (Pin != Pins.Last())
 					{
-						const int32 PinScriptCodeOffset = FunctionContext->GetCodeLocationFromPin(Pin);
-						check(PinScriptCodeOffset != INDEX_NONE);
-						MappedNode->AddLinkedNode(PinScriptCodeOffset, ExecLinkedNode);
+						TunnelExitNode->AddFlags(EScriptExecutionNodeFlags::ReEntrantTunnelPin);
 					}
 				}
+				if (bBranchedExecution)
+				{
+					PinExecNode->AddChildNode(TunnelExitNode);
+				}
+				else
+				{
+					const int32 TunnelId = GetTunnelIdFromName(TunnelExitName);
+					check (TunnelId != INDEX_NONE);
+					PinExecNode->AddLinkedNode(TunnelId, TunnelExitNode);
+				}
+				TArray<TSharedPtr<FScriptExecutionNode>>& TunnelExitPoints = TunnelExitPointMap.FindOrAdd(StagingTunnelName);
+				TunnelExitPoints.Add(TunnelExitNode);
 			}
-			// Evaluate Children for call sites
-			if (UK2Node_CallFunction* FunctionCallSite = Cast<UK2Node_CallFunction>(NodeToMap))
+			else
 			{
-				const UEdGraphNode* FunctionNode = nullptr; 
-				if (UEdGraph* CalledGraph = FunctionCallSite->GetFunctionGraph(FunctionNode))
+				const FName TunnelName(*LinkedPin->PinName);
+				TSharedPtr<FScriptExecutionNode> TunnelEntryNode = MapTunnelEntry(TunnelName, Pin, LinkedPin, PinScriptCodeOffset);
+				if (TunnelEntryNode.IsValid())
 				{
-					// Update Exec node
-					MappedNode->AddFlags(EScriptExecutionNodeFlags::FunctionCall);
-					MappedNode->SetIconColor(FLinearColor(0.46f, 0.54f, 0.95f));
-					MappedNode->SetToolTipText(LOCTEXT("NavigateToFunctionLocationHyperlink_ToolTip", "Navigate to the Function Callsite"));
-					// Update the function context
-					FunctionContext = CreateFunctionContext(CalledGraph->GetFName());
-					// Grab function entries
-					TArray<UK2Node_FunctionEntry*> FunctionEntryNodes;
-					CalledGraph->GetNodesOfClass<UK2Node_FunctionEntry>(FunctionEntryNodes);
-	
-					for (auto EntryNodeIter : FunctionEntryNodes)
+					if (bBranchedExecution)
 					{
-						TSharedPtr<FScriptExecutionNode> FunctionEntry = MapNodeExecution(FunctionContext, EntryNodeIter);
-						MappedNode->AddChildNode(FunctionEntry);
+						PinExecNode->AddChildNode(TunnelEntryNode);
+					}
+					else
+					{
+						PinExecNode->AddLinkedNode(PinScriptCodeOffset, TunnelEntryNode);
 					}
 				}
 			}
-			else if (UK2Node_CustomEvent* CustomEvent = Cast<UK2Node_CustomEvent>(NodeToMap))
+		}
+	}
+}
+
+TSharedPtr<FScriptExecutionNode> FBlueprintFunctionContext::MapTunnelEntry(const FName TunnelName, const UEdGraphPin* ExecPin, const UEdGraphPin* TunnelPin, const int32 PinScriptOffset)
+{
+	// Retrieve tunnel function context
+	const UEdGraphNode* TunnelNode = TunnelPin->GetOwningNode();
+	TSharedPtr<FBlueprintFunctionContext> TunnelContext = GetTunnelContextFromNode(TunnelNode);
+	const int32 TunnelId = TunnelContext->GetTunnelIdFromName(TunnelName);
+	check(TunnelId != INDEX_NONE);
+	// Create tunnel graph instance node
+	FScriptExecNodeParams TunnelParams;
+	TunnelParams.NodeName = FName(*FString::Printf(TEXT("%s_EntryPointId%i"), *TunnelNode->GetFName().ToString(), TunnelId));
+	TunnelParams.ObservedObject = TunnelNode;
+	TunnelParams.DisplayName = TunnelNode->GetNodeTitle(ENodeTitleType::ListView);
+	TunnelParams.Tooltip = LOCTEXT("NavigateToNodeLocationHyperlink_ToolTip", "Navigate to the Node");
+	TunnelParams.NodeFlags = EScriptExecutionNodeFlags::FunctionTunnel;
+	TunnelParams.IconColor = FLinearColor(1.f, 1.f, 1.f, 0.8f);
+	const FSlateBrush* TunnelIcon = TunnelNode->ShowPaletteIconOnNode() ?	FEditorStyle::GetBrush(TunnelNode->GetPaletteIcon(TunnelParams.IconColor)) :
+																			FEditorStyle::GetBrush(TEXT("BlueprintProfiler.BPNode"));
+	TunnelParams.Icon = const_cast<FSlateBrush*>(TunnelIcon);
+	TSharedPtr<FScriptExecutionNode> TunnelCallNode = CreateExecutionNode(TunnelParams);
+	// Add entry point children
+	TSharedPtr<FScriptExecutionNode> EntryPinNode = TunnelContext->GetProfilerDataForNode(TunnelName);
+	check (EntryPinNode.IsValid());
+	TunnelCallNode->AddChildNode(EntryPinNode);
+	if (!EntryPinNode->GetNumChildren())
+	{
+		TSharedPtr<FScriptExecutionNode> FirstTunnelNode = TunnelContext->GetTunnelEntryPoint(TunnelName);
+		check(FirstTunnelNode.IsValid());
+		EntryPinNode->AddChildNode(FirstTunnelNode);
+	}
+	// Map exit points
+	TArray<TSharedPtr<FScriptExecutionNode>> TunnelExitPoints;
+	TunnelContext->GetTunnelExitPoints(TunnelName, TunnelExitPoints);
+	if (const UK2Node_Tunnel* TunnelGraphNode = Cast<UK2Node_Tunnel>(TunnelNode))
+	{
+		TMap<FName, const UEdGraphPin*> NameToPinMap;
+		for (auto ExitPin : TunnelGraphNode->Pins)
+		{
+			if (ExitPin->Direction == EGPD_Output && ExitPin->PinType.PinCategory == UEdGraphSchema_K2::PC_Exec && ExitPin->LinkedTo.Num())
 			{
-				MappedNode->AddFlags(EScriptExecutionNodeFlags::CustomEvent);
+				NameToPinMap.Add(FName(*ExitPin->PinName)) = ExitPin;
+			}
+		}
+		for (auto CurrExitPoint : TunnelExitPoints)
+		{
+			if (const UEdGraphPin** ExitPin = NameToPinMap.Find(CurrExitPoint->GetName()))
+			{
+				for (auto LinkedExitPin : (*ExitPin)->LinkedTo)
+				{
+					TSharedPtr<FScriptExecutionNode> LinkedExecNode = MapNodeExecution(LinkedExitPin->GetOwningNode());
+					if (LinkedExecNode.IsValid())
+					{
+						const int32 ExitScriptCodeOffset = GetCodeLocationFromPin(*ExitPin);
+						if (CurrExitPoint->HasFlags(EScriptExecutionNodeFlags::ReEntrantTunnelPin))
+						{
+							CurrExitPoint->AddLinkedNode(ExitScriptCodeOffset, LinkedExecNode);
+						}
+						else
+						{
+							LinkedExecNode->AddFlags(EScriptExecutionNodeFlags::TunnelThenPin);
+							TunnelCallNode->AddLinkedNode(ExitScriptCodeOffset, LinkedExecNode);
+						}
+					}
+				}
 			}
-			//else if (UK2Node_MacroInstance* MacroCallSite = Cast<UK2Node_MacroInstance>(NodeToMap))
-			//{
-			//	if (UEdGraph* MacroGraph = MacroCallSite->GetMacroGraph())
-			//	{
-			//		// Update Exec node
-			//		MappedNode->AddFlags(EScriptExecutionNodeFlags::MacroCall);
-			//		IconColor = FLinearColor(0.36f, 0.34f, 0.71f);
-			//		ToolTipText = LOCTEXT("NavigateToMacroLocationHyperlink_ToolTip", "Navigate to the Macro Callsite");
-			//		// Grab the macro blueprint context
-			//		UBlueprint* MacroBP = MacroGraph->GetTypedOuter<UBlueprint>();
-			//		UBlueprintGeneratedClass* MacroBPGC = Cast<UBlueprintGeneratedClass>(MacroBP->GeneratedClass);
-			//		TSharedPtr<FBlueprintExecutionContext> MacroBlueprintContext = GetBlueprintContext(MacroBPGC->GetPathName());
-			//		// Grab function entries
-			//		TArray<UK2Node_FunctionEntry*> FunctionEntryNodes;
-			//		MacroGraph->GetNodesOfClass<UK2Node_FunctionEntry>(FunctionEntryNodes);
-	
-			//		for (auto EntryNodeIter : FunctionEntryNodes)
-			//		{
-			//			TSharedPtr<FScriptExecutionNode> MacroEntry = MapNodeExecution(MacroBlueprintContext, EntryNodeIter);
-			//			MappedNode->AddChildNode(MacroEntry);
-			//		}
-			//	}
-			//}
 		}
 	}
-	return MappedNode;
+	// Return the tunnel entry.
+	return TunnelCallNode;
 }
-//////////////////////////////////////////////////////////////////////////
-// FBlueprintFunctionContext
 
 bool FBlueprintFunctionContext::HasProfilerDataForNode(const FName NodeName) const
 {
@@ -601,6 +839,10 @@ const UEdGraphNode* FBlueprintFunctionContext::GetNodeFromCodeLocation(const int
 	{
 		if (const UEdGraphNode* GraphNode = BlueprintClass->GetDebugData().FindSourceNodeFromCodeLocation(Function.Get(), ScriptOffset, true))
 		{
+			if (GraphNode->IsA<UK2Node_MacroInstance>())
+			{
+				GraphNode = BlueprintClass->GetDebugData().FindMacroSourceNodeFromCodeLocation(Function.Get(), ScriptOffset);
+			}
 			Result = GraphNode;
 		}
 	}
@@ -658,6 +900,198 @@ TSharedPtr<FScriptExecutionNode> FBlueprintFunctionContext::CreateExecutionNode(
 	return ScriptExecNode;
 }
 
+void FBlueprintFunctionContext::AddCallSiteEntryPointsToNode(TSharedPtr<FScriptExecutionNode> CallingNode) const
+{
+	if (CallingNode.IsValid())
+	{
+		for (auto EntryPoint : EntryPoints)
+		{
+			for (auto EntryPointChild : EntryPoint->GetChildNodes())
+			{
+				CallingNode->AddChildNode(EntryPointChild);
+			}
+		}
+	}
+}
+
+void FBlueprintFunctionContext::MapTunnelInstance(UK2Node_Tunnel* TunnelInstance)
+{
+	check(TunnelInstance);
+	UEdGraph* TunnelGraph = nullptr;
+	if (UK2Node_MacroInstance* MacroInstance = Cast<UK2Node_MacroInstance>(TunnelInstance))
+	{
+		TunnelGraph = MacroInstance->GetMacroGraph();
+	}
+	else if (UK2Node_Composite* CompositeInstance = Cast<UK2Node_Composite>(TunnelInstance))
+	{
+		TunnelGraph = CompositeInstance->BoundGraph;
+	}
+	TSharedPtr<FBlueprintFunctionContext> TunnelGraphContext;
+	if (TunnelGraph)
+	{
+		TunnelGraphContext = BlueprintContext.Pin()->GetFunctionContext(TunnelGraph->GetFName());
+		if (!TunnelGraphContext.IsValid())
+		{
+			// Create the sub graph context
+			TunnelGraphContext = MakeShareable(new FBlueprintFunctionContext);
+			TunnelGraphContext->FunctionName = TunnelGraph->GetFName();
+			TunnelGraphContext->BlueprintContext = BlueprintContext;
+			TunnelGraphContext->Function = Function;
+			TunnelGraphContext->BlueprintClass = BlueprintClass;
+			BlueprintContext.Pin()->AddFunctionContext(TunnelGraph->GetFName(), TunnelGraphContext);
+			// Map tunnel Inputs
+			TunnelGraphContext->MapTunnelIO(TunnelGraph);
+			// Map tunneling pins into the tunnel graph
+			for (auto EntryPoint : TunnelGraphContext->EntryPoints)
+			{
+				TArray<TSharedPtr<FScriptExecutionNode>>& MappedExits = TunnelGraphContext->TunnelExitPointMap.FindOrAdd(EntryPoint->GetName());
+				if (const UEdGraphPin* EntryPointPin = Cast<UEdGraphPin>(EntryPoint->GetObservedObject()))
+				{
+					TunnelGraphContext->StagingTunnelName = EntryPoint->GetName();
+					for (auto LinkedPin : EntryPointPin->LinkedTo)
+					{
+						const UEdGraphNode* LinkedNode = LinkedPin->GetOwningNode();
+						if (LinkedNode->IsA<UK2Node_Tunnel>())
+						{
+							// This handles input linked directly to output
+							TSharedPtr<FScriptExecutionNode> OutputPinNode = TunnelGraphContext->GetProfilerDataForNode(FName(*LinkedPin->PinName));
+							check(OutputPinNode.IsValid());
+							MappedExits.Add(OutputPinNode);
+						}
+						else
+						{
+							// This adds any encountered exit points to the active tunnel name.
+							TSharedPtr<FScriptExecutionNode> ExecNode = TunnelGraphContext->MapNodeExecution(LinkedPin->GetOwningNode());
+							if (ExecNode.IsValid())
+							{
+								TunnelGraphContext->TunnelEntryPointMap.Add(EntryPoint->GetName()) = ExecNode;
+							}
+						}
+					}
+				}
+			}
+			TunnelGraphContext->StagingTunnelName = NAME_None;
+		}
+		// Patch info into the calling function context
+		for (auto ExecutionNode : TunnelGraphContext->ExecutionNodes)
+		{
+			const int32 TunnelId = TunnelGraphContext->GetTunnelIdFromName(ExecutionNode.Value->GetName());
+			if (TunnelId == INDEX_NONE)
+			{
+				TSharedPtr<FScriptExecutionNode>* ExistingNode = ExecutionNodes.Find(ExecutionNode.Key);
+				if (ExistingNode == nullptr)
+				{
+					ExecutionNodes.Add(ExecutionNode.Key) = ExecutionNode.Value;
+				}
+			}
+		}
+	}
+}
+
+TSharedPtr<FBlueprintFunctionContext> FBlueprintFunctionContext::GetTunnelContextFromNode(const UEdGraphNode* TunnelNode) const
+{
+	TSharedPtr<FBlueprintFunctionContext> TunnelContext;
+	if (const UK2Node_Composite* CompositeNode = Cast<UK2Node_Composite>(TunnelNode))
+	{
+		TunnelContext = BlueprintContext.Pin()->GetFunctionContext(CompositeNode->BoundGraph->GetFName());
+	}
+	else if (const UK2Node_MacroInstance* MacroNode = Cast<UK2Node_MacroInstance>(TunnelNode))
+	{
+		TunnelContext = BlueprintContext.Pin()->GetFunctionContext(MacroNode->GetMacroGraph()->GetFName());
+	}
+	return TunnelContext;
+}
+
+TSharedPtr<FScriptExecutionNode> FBlueprintFunctionContext::GetTunnelEntryPoint(const FName TunnelName) const
+{
+	TSharedPtr<FScriptExecutionNode> EntryPoint;
+	if (const TSharedPtr<FScriptExecutionNode>* FoundEntryPoint = TunnelEntryPointMap.Find(TunnelName))
+	{
+		EntryPoint = *FoundEntryPoint;
+	}
+	return EntryPoint;
+}
+
+void FBlueprintFunctionContext::GetTunnelExitPoints(const FName TunnelName, TArray<TSharedPtr<FScriptExecutionNode>>& ExitPointsOut) const
+{
+	ExitPointsOut.Reset();
+	if (const TArray<TSharedPtr<FScriptExecutionNode>>* FoundExitPoints = TunnelExitPointMap.Find(TunnelName))
+	{
+		ExitPointsOut = *FoundExitPoints;
+	}
+}
+
+void FBlueprintFunctionContext::MapTunnelIO(UEdGraph* TunnelGraph)
+{
+	// Map graph entry points
+	check (TunnelGraph);
+	TArray<UK2Node_Tunnel*> GraphTunnels;
+	TunnelGraph->GetNodesOfClass<UK2Node_Tunnel>(GraphTunnels);
+	for (auto Tunnel : GraphTunnels)
+	{
+		for (auto Pin : Tunnel->Pins)
+		{
+			if (Pin->LinkedTo.Num() && Pin->PinType.PinCategory == UEdGraphSchema_K2::PC_Exec)
+			{
+				if (Pin->Direction == EGPD_Output)
+				{
+					// Create entry point
+					FScriptExecNodeParams LinkParams;
+					LinkParams.NodeName = FName(*Pin->PinName);
+					LinkParams.ObservedObject = Pin;
+					LinkParams.DisplayName = Pin->GetDisplayName().IsEmpty() ? FText::FromString(*Pin->PinName) : Pin->GetDisplayName();
+					LinkParams.Tooltip = LOCTEXT("ExecPin_ToolTip", "Expand tunnel entry point");
+					LinkParams.NodeFlags = EScriptExecutionNodeFlags::TunnelPin;
+					LinkParams.IconColor = FLinearColor(1.f, 1.f, 1.f, 0.8f);
+					const FSlateBrush* Icon = FEditorStyle::GetBrush(TEXT("BlueprintProfiler.BPPinConnected"));
+					LinkParams.Icon = const_cast<FSlateBrush*>(Icon);
+					TSharedPtr<FScriptExecutionNode> EntryPoint = CreateExecutionNode(LinkParams);
+					AddEntryPoint(EntryPoint);
+				}
+				else if (Pin->Direction == EGPD_Input)
+				{
+					// Create exit point
+					FScriptExecNodeParams LinkParams;
+					LinkParams.NodeName = FName(*Pin->PinName);
+					LinkParams.ObservedObject = Pin;
+					LinkParams.DisplayName = Pin->GetDisplayName().IsEmpty() ? FText::FromString(*Pin->PinName) : Pin->GetDisplayName();
+					LinkParams.Tooltip = LOCTEXT("ExecPin_ToolTip", "Expand tunnel exit point");
+					LinkParams.NodeFlags = EScriptExecutionNodeFlags::TunnelExitPin;
+					LinkParams.IconColor = FLinearColor(1.f, 1.f, 1.f, 0.8f);
+					const FSlateBrush* Icon = FEditorStyle::GetBrush(TEXT("BlueprintProfiler.BPPinConnected"));
+					LinkParams.Icon = const_cast<FSlateBrush*>(Icon);
+					TSharedPtr<FScriptExecutionNode> ExitPoint = CreateExecutionNode(LinkParams);
+					AddExitPoint(ExitPoint);
+				}
+			}
+		}
+	}
+}
+
+int32 FBlueprintFunctionContext::GetTunnelIdFromName(const FName TunnelName) const
+{
+	int32 TunnelId = INDEX_NONE;
+	if (const int32* FoundTunnelId = TunnelIds.Find(TunnelName))
+	{
+		TunnelId = *FoundTunnelId;
+	}
+	return TunnelId;
+}
+
+void FBlueprintFunctionContext::AddEntryPoint(TSharedPtr<FScriptExecutionNode> EntryPoint)
+{
+	EntryPoints.Add(EntryPoint);
+	const int32 TunnelId = TunnelIds.Num();
+	TunnelIds.Add(EntryPoint->GetName()) = TunnelId;
+}
+
+void FBlueprintFunctionContext::AddExitPoint(TSharedPtr<FScriptExecutionNode> ExitPoint)
+{
+	ExitPoints.Add(ExitPoint);
+	const int32 TunnelId = TunnelIds.Num();
+	TunnelIds.Add(ExitPoint->GetName()) = TunnelId;
+}
+
 //////////////////////////////////////////////////////////////////////////
 // FScriptEventPlayback
 
@@ -671,6 +1105,12 @@ bool FScriptEventPlayback::Process(const TArray<FScriptInstrumentedEvent>& Signa
 		TSharedPtr<FBlueprintFunctionContext> FunctionContext;
 		TArray<FScriptInstrumentedEvent> UnproccessedEvents;
 	};
+	struct TunnelEventHelper
+	{
+		double TunnelEntryTime;
+		FTracePath EntryTracePath;
+		TSharedPtr<FScriptExecutionNode> TunnelEntryPoint;
+	};
 
 	const int32 NumEvents = (StopIdx+1) - StartIdx;
 	ProcessingState = EEventProcessingResult::Failed;
@@ -685,6 +1125,7 @@ bool FScriptEventPlayback::Process(const TArray<FScriptInstrumentedEvent>& Signa
 		TSharedPtr<FScriptExecutionNode> EventNode = FunctionContext->GetProfilerDataForNode(EventName);
 		// Find Associated graph nodes and submit into a node map for later processing.
 		TMap<const UEdGraphNode*, PerNodeEventHelper> NodeMap;
+		TArray<TunnelEventHelper> Tunnels;
 		ProcessingState = EEventProcessingResult::Success;
 		int32 LastEventIdx = StartIdx;
 		const int32 EventStartOffset = SignalData[StartIdx].IsResumeEvent() ? 3 : 1;
@@ -706,6 +1147,7 @@ bool FScriptEventPlayback::Process(const TArray<FScriptInstrumentedEvent>& Signa
 				{
 					NodeEntry.FunctionContext = FunctionContext;
 					NodeEntry.ExecNode = FunctionContext->GetProfilerDataForNode(GraphNode->GetFName());
+					check(NodeEntry.ExecNode.IsValid());
 				}
 				switch (CurrSignal.GetType())
 				{
@@ -713,7 +1155,7 @@ bool FScriptEventPlayback::Process(const TArray<FScriptInstrumentedEvent>& Signa
 					{
 						if (const UEdGraphPin* Pin = FunctionContext->GetPinFromCodeLocation(CurrSignal.GetScriptCodeOffset()))
 						{
-							if (NodeEntry.FunctionContext->HasProfilerDataForNode(Pin->GetOwningNode()->GetFName()))
+							if (FunctionContext->HasProfilerDataForNode(Pin->GetOwningNode()->GetFName()))
 							{
 								TracePath.AddExitPin(CurrSignal.GetScriptCodeOffset());
 								NodeEntry.InputTracePaths.Insert(TracePath, 0);
@@ -808,8 +1250,48 @@ bool FScriptEventPlayback::Process(const TArray<FScriptInstrumentedEvent>& Signa
 								TracePath = NodeEntry.EntryTracePaths.Last();
 							}
 						}
-						TracePath.AddExitPin(CurrSignal.GetScriptCodeOffset());
+						const int32 ScriptCodeExit = CurrSignal.GetScriptCodeOffset();
+						// Process tunnel timings
+						for (auto LinkIter : NodeEntry.ExecNode->GetLinkedNodes())
+						{
+							if (LinkIter.Key == ScriptCodeExit && LinkIter.Value->HasFlags(EScriptExecutionNodeFlags::FunctionTunnel))
+							{
+								TunnelEventHelper NewTunnel;
+								NewTunnel.EntryTracePath = TracePath;
+								NewTunnel.EntryTracePath.AddExitPin(ScriptCodeExit);
+								NewTunnel.TunnelEntryPoint = LinkIter.Value;
+								NewTunnel.TunnelEntryTime = CurrSignal.GetTime();
+								Tunnels.Push(NewTunnel);
+							}
+							else if (LinkIter.Value->HasFlags(EScriptExecutionNodeFlags::TunnelExitPin))
+							{
+								if (Tunnels.Num())
+								{
+									TunnelEventHelper Tunnel = Tunnels.Pop();
+									TSharedPtr<FScriptPerfData> PerfData = Tunnel.TunnelEntryPoint->GetPerfDataByInstanceAndTracePath(InstanceName, Tunnel.EntryTracePath);
+									PerfData->AddEventTiming(CurrSignal.GetTime() - Tunnel.TunnelEntryTime);
+									FTracePath ExitTrace(TracePath);
+									ExitTrace.AddExitPin(LinkIter.Key);
+									PerfData = LinkIter.Value->GetPerfDataByInstanceAndTracePath(InstanceName, ExitTrace);
+									PerfData->AddEventTiming(0.0);
+								}
+							}
+						}
+						TracePath.AddExitPin(ScriptCodeExit);
 						NodeEntry.UnproccessedEvents.Add(CurrSignal);
+						// Process cyclic linkage - reset to root link tracepath
+						TSharedPtr<FScriptExecutionNode> NextLink = NodeEntry.ExecNode->GetLinkedNodeByScriptOffset(ScriptCodeExit);
+						if (NextLink.IsValid() && NextLink->HasFlags(EScriptExecutionNodeFlags::CyclicLinkage))
+						{
+							const UEdGraphNode* LinkedGraphNode = Cast<UEdGraphNode>(NextLink->GetObservedObject());
+							if (PerNodeEventHelper* LinkedEntry = NodeMap.Find(LinkedGraphNode))
+							{
+								if (LinkedEntry->EntryTracePaths.Num())
+								{
+									TracePath = LinkedEntry->EntryTracePaths.Last();
+								}
+							}
+						}
 						break;
 					}
 					default:
@@ -842,7 +1324,7 @@ bool FScriptEventPlayback::Process(const TArray<FScriptInstrumentedEvent>& Signa
 				EventNode = FunctionContext->GetProfilerDataForNode(EventTiming.Key);
 				check (EventNode.IsValid());
 				TSharedPtr<FScriptPerfData> PerfData = EventNode->GetPerfDataByInstanceAndTracePath(InstanceName, EventTracePath);
-				PerfData->AddEventTiming(0.0, SignalData[StopIdx].GetTime() - SignalData[LastEventIdx].GetTime());
+				PerfData->AddEventTiming(SignalData[StopIdx].GetTime() - SignalData[LastEventIdx].GetTime());
 			}
 			EventTimings.Reset();
 		}
@@ -850,42 +1332,54 @@ bool FScriptEventPlayback::Process(const TArray<FScriptInstrumentedEvent>& Signa
 		for (auto NodeEntry : NodeMap)
 		{
 			TSharedPtr<FScriptExecutionNode> ExecNode = NodeEntry.Value.ExecNode;
+			TSharedPtr<FScriptExecutionNode> PureNode;
 			double PureNodeEntryTime = 0.0;
-			double LastPureNodeTime = 0.0;
+			double PureChainEntryTime = 0.0;
 			double NodeEntryTime = 0.0;
 			int32 TracePathIdx = 0;
 			for (auto EventIter = NodeEntry.Value.UnproccessedEvents.CreateIterator(); EventIter; ++EventIter)
 			{
 				switch(EventIter->GetType())
 				{
-					case EScriptInstrumentation::PushState:
-					{
-						LastPureNodeTime = EventIter->GetTime();
-						break;
-					}
 					case EScriptInstrumentation::PureNodeEntry:
 					{
-						if (PureNodeEntryTime == 0.0)
+						if (PureChainEntryTime == 0.0)
+						{
+							PureChainEntryTime = EventIter->GetTime();
+						}
+						else if (PureNode.IsValid())
 						{
-							PureNodeEntryTime = EventIter->GetTime();
+							TSharedPtr<FScriptPerfData> PerfData = PureNode->GetPerfDataByInstanceAndTracePath(InstanceName, NodeEntry.Value.InputTracePaths.Pop());
+							PerfData->AddEventTiming(EventIter->GetTime() - PureNodeEntryTime);
 						}
 
+						PureNode.Reset();
+						PureNodeEntryTime = EventIter->GetTime();
+						
+						FunctionContext = BlueprintContext->GetFunctionContext(EventIter->GetFunctionName());
+						check(FunctionContext.IsValid());
+
 						const int32 ScriptCodeOffset = EventIter->GetScriptCodeOffset();
 						if (const UEdGraphPin* Pin = FunctionContext->GetPinFromCodeLocation(ScriptCodeOffset))
 						{
 							if (FunctionContext->HasProfilerDataForNode(Pin->GetOwningNode()->GetFName()))
 							{
-								TSharedPtr<FScriptExecutionNode> PureNode = NodeEntry.Value.FunctionContext->GetProfilerDataForNode(Pin->GetOwningNode()->GetFName());
-								TSharedPtr<FScriptPerfData> PerfData = PureNode->GetPerfDataByInstanceAndTracePath(InstanceName, NodeEntry.Value.InputTracePaths.Pop());
-								PerfData->AddEventTiming(0.0f, EventIter->GetTime() - LastPureNodeTime);
+								PureNode = NodeEntry.Value.FunctionContext->GetProfilerDataForNode(Pin->GetOwningNode()->GetFName());
 							}
 						}
-
-						LastPureNodeTime = EventIter->GetTime();
 						break;
 					}
 					case EScriptInstrumentation::NodeEntry:
 					{
+						if (PureNode.IsValid())
+						{
+							TSharedPtr<FScriptPerfData> PerfData = PureNode->GetPerfDataByInstanceAndTracePath(InstanceName, NodeEntry.Value.InputTracePaths.Pop());
+							PerfData->AddEventTiming(EventIter->GetTime() - PureNodeEntryTime);
+
+							PureNode.Reset();
+							PureNodeEntryTime = 0.0;
+						}
+
 						if (NodeEntryTime == 0.0)
 						{
 							NodeEntryTime = EventIter->GetTime();
@@ -897,10 +1391,17 @@ bool FScriptEventPlayback::Process(const TArray<FScriptInstrumentedEvent>& Signa
 						check(NodeEntryTime != 0.0 );
 						const FTracePath NodeTracePath = NodeEntry.Value.EntryTracePaths[TracePathIdx++];
 						TSharedPtr<FScriptPerfData> PerfData = ExecNode->GetPerfDataByInstanceAndTracePath(InstanceName, NodeTracePath);
-						double PureNodeDuration = PureNodeEntryTime != 0.0 ? (NodeEntryTime - PureNodeEntryTime) : 0.0;
+						double PureChainDuration = PureChainEntryTime != 0.0 ? (NodeEntryTime - PureChainEntryTime) : 0.0;
 						double NodeDuration = EventIter->GetTime() - NodeEntryTime;
-						PerfData->AddEventTiming(PureNodeDuration, NodeDuration);
-						PureNodeEntryTime = NodeEntryTime = 0.0;
+						PerfData->AddEventTiming(NodeDuration);
+						PureChainEntryTime = NodeEntryTime = 0.0;
+						PerfData->AddPureChainTiming(PureChainDuration);
+						TSharedPtr<FScriptExecutionNode> PureChainNode = ExecNode->GetPureChainNode();
+						if (PureChainNode.IsValid())
+						{
+							PerfData = PureChainNode->GetPerfDataByInstanceAndTracePath(InstanceName, NodeTracePath);
+							PerfData->AddEventTiming(PureChainDuration);
+						}
 						break;
 					}
 				}
diff --git a/Engine/Source/Developer/BlueprintProfiler/Public/ScriptInstrumentationPlayback.h b/Engine/Source/Developer/BlueprintProfiler/Public/ScriptInstrumentationPlayback.h
--- a/Engine/Source/Developer/BlueprintProfiler/Public/ScriptInstrumentationPlayback.h
+++ b/Engine/Source/Developer/BlueprintProfiler/Public/ScriptInstrumentationPlayback.h
@@ -26,6 +26,9 @@ public:
 	/** Returns the blueprint exec node for this context */
 	TSharedPtr<class FScriptExecutionBlueprint> GetBlueprintExecNode() const { return BlueprintNode; }
 
+	/** Add event execution node */
+	void AddEventNode(TSharedPtr<FScriptExecutionNode> EventExecNode);
+
 	/** Returns true if an execution node exists for the instance */
 	BLUEPRINTPROFILER_API bool HasProfilerDataForInstance(const FName InstanceName) const;
 
@@ -41,6 +44,12 @@ public:
 	/** Returns the function context matching the supplied name providing it exists */
 	TSharedPtr<class FBlueprintFunctionContext> GetFunctionContext(const FName FunctionNameIn) const;
 
+	/** Adds a manually created function context */
+	void AddFunctionContext(const FName FunctionName, TSharedPtr<class FBlueprintFunctionContext> FunctionContext)
+	{
+		FunctionContexts.Add(FunctionName) = FunctionContext; 
+	}
+
 	/** Returns if this context has an execution node for the specified graph node */
 	BLUEPRINTPROFILER_API bool HasProfilerDataForNode(const UEdGraphNode* GraphNode) const;
 
@@ -53,17 +62,11 @@ public:
 private:
 
 	/** Creates and returns a new function context, returns existing if present  */
-	TSharedPtr<FBlueprintFunctionContext> CreateFunctionContext(const FName FunctionNameIn);
+	TSharedPtr<FBlueprintFunctionContext> CreateFunctionContext(UEdGraph* Graph);
 
 	/** Walks the blueprint and maps events and functions ready for profiling data */
 	bool MapBlueprintExecution();
 
-	/** Locates and generates event contexts for any input events passed in */
-	void CreateInputEvents(const TArray<UK2Node*>& InputEventNodes);
-
-	/** Maps each blueprint node following execution wires */
-	TSharedPtr<FScriptExecutionNode> MapNodeExecution(TSharedPtr<FBlueprintFunctionContext> FunctionContext, UEdGraphNode* NodeToMap);
-
 private:
 
 	/** Blueprint the context represents */
@@ -90,11 +93,11 @@ class FBlueprintFunctionContext : public TSharedFromThis<FBlueprintFunctionConte
 {
 public:
 
-	FBlueprintFunctionContext(UFunction* FunctionIn, TWeakObjectPtr<UBlueprintGeneratedClass> BPGC)
-		: Function(FunctionIn)
-		, BlueprintClass(BPGC)
-	{
-	}
+	/** Initialise the function context from the graph */
+	void InitialiseContextFromGraph(TSharedPtr<FBlueprintExecutionContext> BlueprintContext, UEdGraph* Graph);
+
+	/** Map Function */
+	void MapFunction();
 
 	/** Returns if the function context is fully formed */
 	bool IsContextValid() const { return Function.IsValid() && BlueprintClass.IsValid(); }
@@ -120,15 +123,71 @@ public:
 	/** Returns the execution node representing the node name */
 	TSharedPtr<FScriptExecutionNode> GetProfilerDataForNode(const FName NodeName);
 
-protected:
+	/** Adds all function entry points to node as children */
+	void AddCallSiteEntryPointsToNode(TSharedPtr<FScriptExecutionNode> CallingNode) const;
+
+	/** Returns the function name for this context */
+	FName GetFunctionName() const { return FunctionName; }
+
+private:
 
 	friend FBlueprintExecutionContext;
 
 	/** Utility to create execution node */
 	TSharedPtr<FScriptExecutionNode> CreateExecutionNode(FScriptExecNodeParams& InitParams);
 
+	/** Adds and links in a new function entry point */
+	void AddEntryPoint(TSharedPtr<FScriptExecutionNode> EntryPoint);
+
+	/** Adds a new function exit point */
+	void AddExitPoint(TSharedPtr<FScriptExecutionNode> ExitPoint);
+
+	/** Locates and generates event contexts for any input events passed in */
+	void CreateInputEvents(TSharedPtr<FBlueprintExecutionContext> BlueprintContextIn, const TArray<UK2Node*>& InputEventNodes);
+
+	/** Processes and detects any cyclic links making the linkage safe for traversal */
+	bool DetectCyclicLinks(TSharedPtr<FScriptExecutionNode> ExecNode, TSet<TSharedPtr<FScriptExecutionNode>>& Filter);
+
+	// --Execution mapping functionality
+
+	/** Maps each blueprint node following execution wires */
+	TSharedPtr<FScriptExecutionNode> MapNodeExecution(UEdGraphNode* NodeToMap);
+
+	/** Maps input pin execution */
+	void MapInputPins(TSharedPtr<FScriptExecutionNode> ExecNode, const TArray<UEdGraphPin*>& Pins);
+
+	/** Maps pin execution */
+	void MapExecPins(TSharedPtr<FScriptExecutionNode> ExecNode, const TArray<UEdGraphPin*>& Pins);
+
+	/** Maps execution pin tunnel entry */
+	TSharedPtr<FScriptExecutionNode> MapTunnelEntry(const FName TunnelName, const UEdGraphPin* ExecPin, const UEdGraphPin* TunnelPin, const int32 PinScriptOffset);
+
+	// --Tunnel mapping functionality
+
+	/** Maps the tunnel point into the instanced graph, creating the instanced graph if not already existing. */
+	void MapTunnelInstance(UK2Node_Tunnel* TunnelInstance);
+
+	/** Returns the tunnel function context that the tunnel node invokes */
+	TSharedPtr<FBlueprintFunctionContext> GetTunnelContextFromNode(const UEdGraphNode* TunnelNode) const;
+
+	/** Returns the registered tunnel execution entry point given the tunnel name */
+	TSharedPtr<FScriptExecutionNode> GetTunnelEntryPoint(const FName TunnelName) const;
+
+	/** Returns the registered tunnel exit points given the tunnel name */
+	void GetTunnelExitPoints(const FName TunnelName, TArray<TSharedPtr<FScriptExecutionNode>>& ExitPointsOut) const;
+
+	/** Maps input and output pins in a tunnel graph */
+	void MapTunnelIO(UEdGraph* TunnelGraph);
+
+	/** Returns the tunnel Id associated with the tunnel name */
+	int32 GetTunnelIdFromName(const FName TunnelName) const;
+
 private:
 
+	/** The function name */
+	FName FunctionName;
+	/** The Owning Blueprint context */
+	TWeakPtr<FBlueprintExecutionContext> BlueprintContext;
 	/** The UFunction this context represents */
 	TWeakObjectPtr<UFunction> Function;
 	/** The blueprint generated class for the owning blueprint */
@@ -139,6 +198,18 @@ private:
 	TMap<int32, TWeakObjectPtr<const UEdGraphPin>> ScriptOffsetToPins;
 	/** Execution nodes containing profiling data */
 	TMap<FName, TSharedPtr<FScriptExecutionNode>> ExecutionNodes;
+	/** Graph entry points */
+	TArray<TSharedPtr<FScriptExecutionNode>> EntryPoints;
+	/** Graph exit points */
+	TArray<TSharedPtr<FScriptExecutionNode>> ExitPoints;
+	/** Map associating function entry points to execution points */
+	TMap<FName, TSharedPtr<FScriptExecutionNode>> TunnelEntryPointMap;
+	/** Map associating function entry points to exit points */
+	TMap<FName, TArray<TSharedPtr<FScriptExecutionNode>>> TunnelExitPointMap;
+	/** Staging tunnel name during mapping */
+	FName StagingTunnelName;
+	/** Map of tunnel Ids */
+	TMap<FName, int32> TunnelIds;
 };
 
 //////////////////////////////////////////////////////////////////////////
diff --git a/Engine/Source/Editor/BlueprintGraph/Classes/EdGraphSchema_K2.h b/Engine/Source/Editor/BlueprintGraph/Classes/EdGraphSchema_K2.h
--- a/Engine/Source/Editor/BlueprintGraph/Classes/EdGraphSchema_K2.h
+++ b/Engine/Source/Editor/BlueprintGraph/Classes/EdGraphSchema_K2.h
@@ -172,6 +172,13 @@ public:
 
 	static const FName MD_ArrayParam;
 	static const FName MD_ArrayDependentParam;
+
+	/** Metadata that identifies an integral property as a bitmask. */
+	static const FName MD_Bitmask;
+	/** Metadata that associates a bitmask property with a bitflag enum. */
+	static const FName MD_BitmaskEnum;
+	/** Metadata that identifies an enum as a set of explicitly-named bitflags. */
+	static const FName MD_Bitflags;
 	
 private:
 	// This class should never be instantiated
@@ -242,6 +249,7 @@ class BLUEPRINTGRAPH_API UEdGraphSchema_K2 : public UEdGraphSchema
 	static const FString PSC_Self;    // Category=PC_Object or PC_Class, indicates the class being compiled
 
 	static const FString PSC_Index;	// Category=PC_Wildcard, indicates the wildcard will only accept Int, Bool, Byte and Enum pins (used when a pin represents indexing a list)
+	static const FString PSC_Bitmask;	// Category=PC_Byte or PC_Int, indicates that the pin represents a bitmask field. SubCategoryObject is either NULL or the UEnum object to which the bitmap is linked for bitflag name specification.
 
 	// Pin names that have special meaning and required types in some contexts (depending on the node type)
 	static const FString PN_Execute;    // Category=PC_Exec, singleton, input
@@ -417,7 +425,7 @@ public:
 	virtual void GetAssetsPinHoverMessage(const TArray<FAssetData>& Assets, const UEdGraphPin* HoverPin, FString& OutTooltipText, bool& OutOkIcon) const override;
 	virtual bool CanDuplicateGraph(UEdGraph* InSourceGraph) const override;
 	virtual UEdGraph* DuplicateGraph(UEdGraph* GraphToDuplicate) const override;
-	virtual UEdGraphNode* CreateSubstituteNode(UEdGraphNode* Node, const UEdGraph* Graph, FObjectInstancingGraph* InstanceGraph) const override;
+	virtual UEdGraphNode* CreateSubstituteNode(UEdGraphNode* Node, const UEdGraph* Graph, FObjectInstancingGraph* InstanceGraph, TArray<FName>& InOutExtraNames) const override;
 	virtual int32 GetNodeSelectionCount(const UEdGraph* Graph) const override;
 	virtual TSharedPtr<FEdGraphSchemaAction> GetCreateCommentAction() const override;
 	virtual bool FadeNodeWhenDraggingOffPin(const UEdGraphNode* Node, const UEdGraphPin* Pin) const override;
diff --git a/Engine/Source/Editor/BlueprintGraph/Classes/K2Node.h b/Engine/Source/Editor/BlueprintGraph/Classes/K2Node.h
--- a/Engine/Source/Editor/BlueprintGraph/Classes/K2Node.h
+++ b/Engine/Source/Editor/BlueprintGraph/Classes/K2Node.h
@@ -162,7 +162,11 @@ class UK2Node : public UEdGraphNode
 	UPROPERTY()
 	TMap<UEdGraphPin*, UEdGraphPin*> ByRefMatchupPins;
 
-	// Reallocate pins during reconstruction; by default ignores the old pins and calls AllocateDefaultPins()
+	/**
+	 * Reallocate pins during reconstruction; by default ignores the old pins and calls AllocateDefaultPins()
+	 * If you override this to create additional pins you likely need to call RestoreSplitPins to restore any
+	 * pins that have been split (e.g. a vector pin split into its components)
+	 */
 	BLUEPRINTGRAPH_API virtual void ReallocatePinsDuringReconstruction(TArray<UEdGraphPin*>& OldPins);
 
 	/** Returns whether this node is considered 'pure' by the compiler */
@@ -374,6 +378,9 @@ protected:
 	 */
 	BLUEPRINTGRAPH_API ERedirectType ShouldRedirectParam(const TArray<FString>& OldPinNames, FName& NewPinName, const UK2Node * NewPinNode) const;
 
+	// Helper function to restore Split Pins after ReallocatePinsDuringReconstruction, call after recreating all pins to restore split pin state
+	BLUEPRINTGRAPH_API void RestoreSplitPins(TArray<UEdGraphPin*>& OldPins);
+
 	/** 
 	 * Sends a message to the owning blueprint's CurrentMessageLog, if there is one available.  Otherwise, defaults to logging to the normal channels.
 	 * Should use this for node actions that happen during compilation!
diff --git a/Engine/Source/Editor/BlueprintGraph/Classes/K2Node_ActorBoundEvent.h b/Engine/Source/Editor/BlueprintGraph/Classes/K2Node_ActorBoundEvent.h
--- a/Engine/Source/Editor/BlueprintGraph/Classes/K2Node_ActorBoundEvent.h
+++ b/Engine/Source/Editor/BlueprintGraph/Classes/K2Node_ActorBoundEvent.h
@@ -30,6 +30,7 @@ class UK2Node_ActorBoundEvent : public UK2Node_Event
 	//~ Begin UEdGraphNode Interface
 	virtual void ReconstructNode() override;
 	virtual void DestroyNode() override;
+	virtual bool CanPasteHere(const UEdGraph* TargetGraph) const override;
 	virtual FText GetNodeTitle(ENodeTitleType::Type TitleType) const override;
 	virtual FText GetTooltipText() const override;
 	virtual FString GetDocumentationLink() const override;
diff --git a/Engine/Source/Editor/BlueprintGraph/Classes/K2Node_ComponentBoundEvent.h b/Engine/Source/Editor/BlueprintGraph/Classes/K2Node_ComponentBoundEvent.h
--- a/Engine/Source/Editor/BlueprintGraph/Classes/K2Node_ComponentBoundEvent.h
+++ b/Engine/Source/Editor/BlueprintGraph/Classes/K2Node_ComponentBoundEvent.h
@@ -31,6 +31,7 @@ class UK2Node_ComponentBoundEvent : public UK2Node_Event
 	//~ End UObject Interface
 
 	//~ Begin UEdGraphNode Interface
+	virtual bool CanPasteHere(const UEdGraph* TargetGraph) const override;
 	virtual FText GetNodeTitle(ENodeTitleType::Type TitleType) const override;
 	virtual FText GetTooltipText() const override;
 	virtual FString GetDocumentationLink() const override;
diff --git a/Engine/Source/Editor/BlueprintGraph/Classes/K2Node_ConstructObjectFromClass.h b/Engine/Source/Editor/BlueprintGraph/Classes/K2Node_ConstructObjectFromClass.h
--- a/Engine/Source/Editor/BlueprintGraph/Classes/K2Node_ConstructObjectFromClass.h
+++ b/Engine/Source/Editor/BlueprintGraph/Classes/K2Node_ConstructObjectFromClass.h
@@ -20,6 +20,7 @@ class BLUEPRINTGRAPH_API UK2Node_ConstructObjectFromClass : public UK2Node
 	virtual bool HasExternalDependencies(TArray<class UStruct*>* OptionalOutput) const override;
 	virtual bool IsCompatibleWithGraph(const UEdGraph* TargetGraph) const override;
 	virtual void PinConnectionListChanged(UEdGraphPin* Pin);
+	virtual void GetPinHoverText(const UEdGraphPin& Pin, FString& HoverTextOut) const override;
 	//~ End UEdGraphNode Interface.
 
 	//~ Begin UK2Node Interface
diff --git a/Engine/Source/Editor/BlueprintGraph/Classes/K2Node_SpawnActorFromClass.h b/Engine/Source/Editor/BlueprintGraph/Classes/K2Node_SpawnActorFromClass.h
--- a/Engine/Source/Editor/BlueprintGraph/Classes/K2Node_SpawnActorFromClass.h
+++ b/Engine/Source/Editor/BlueprintGraph/Classes/K2Node_SpawnActorFromClass.h
@@ -17,6 +17,7 @@ class BLUEPRINTGRAPH_API UK2Node_SpawnActorFromClass : public UK2Node
 	virtual FText GetNodeTitle(ENodeTitleType::Type TitleType) const override;
 	virtual void PinDefaultValueChanged(UEdGraphPin* Pin) override;
 	virtual void PinConnectionListChanged(UEdGraphPin* Pin) override;
+	virtual void GetPinHoverText(const UEdGraphPin& Pin, FString& HoverTextOut) const override;
 	virtual FText GetTooltipText() const override;
 	virtual void ExpandNode(class FKismetCompilerContext& CompilerContext, UEdGraph* SourceGraph) override;
 	virtual bool HasExternalDependencies(TArray<class UStruct*>* OptionalOutput) const override;
diff --git a/Engine/Source/Editor/BlueprintGraph/Private/BlueprintEventNodeSpawner.cpp b/Engine/Source/Editor/BlueprintGraph/Private/BlueprintEventNodeSpawner.cpp
--- a/Engine/Source/Editor/BlueprintGraph/Private/BlueprintEventNodeSpawner.cpp
+++ b/Engine/Source/Editor/BlueprintGraph/Private/BlueprintEventNodeSpawner.cpp
@@ -13,16 +13,6 @@
 
 namespace UBlueprintEventNodeSpawnerImpl
 {
-	/**
-	 * Helper function for scanning a blueprint for a certain, custom named,
-	 * event.
-	 * 
-	 * @param  Blueprint	The blueprint you want to look through
-	 * @param  CustomName	The event name you want to check for.
-	 * @return Null if no event was found, otherwise a pointer to the named event.
-	 */
-	static UK2Node_Event* FindCustomEventNode(UBlueprint* Blueprint, FName const CustomName);
-
 	/**
 	 * Helper function for removing all disabled nodes connected to a 
 	 * disabled source node
@@ -33,28 +23,6 @@ namespace UBlueprintEventNodeSpawnerImpl
 	static void RemoveAllDisabledNodes(UEdGraphNode* InNode, UEdGraph* InParentGraph);
 }
 
-//------------------------------------------------------------------------------
-UK2Node_Event* UBlueprintEventNodeSpawnerImpl::FindCustomEventNode(UBlueprint* Blueprint, FName const CustomName)
-{
-	UK2Node_Event* FoundNode = nullptr;
-
-	if (CustomName != NAME_None)
-	{
-		TArray<UK2Node_Event*> AllEvents;
-		FBlueprintEditorUtils::GetAllNodesOfClass<UK2Node_Event>(Blueprint, AllEvents);
-
-		for (UK2Node_Event* EventNode : AllEvents)
-		{
-			if (EventNode->CustomFunctionName == CustomName)
-			{
-				FoundNode = EventNode;
-				break;
-			}
-		}
-	}	
-	return FoundNode;
-}
-
 //------------------------------------------------------------------------------
 static void UBlueprintEventNodeSpawnerImpl::RemoveAllDisabledNodes(UEdGraphNode* InNode, UEdGraph* InParentGraph)
 {
@@ -239,7 +207,7 @@ UK2Node_Event const* UBlueprintEventNodeSpawner::FindPreExistingEvent(UBlueprint
 	check(Blueprint != nullptr);
 	if (IsForCustomEvent())
 	{
-		PreExistingNode = UBlueprintEventNodeSpawnerImpl::FindCustomEventNode(Blueprint, CustomEventName);
+		PreExistingNode = FBlueprintEditorUtils::FindCustomEventNode(Blueprint, CustomEventName);
 	}
 	else
 	{
diff --git a/Engine/Source/Editor/BlueprintGraph/Private/EdGraphSchema_K2.cpp b/Engine/Source/Editor/BlueprintGraph/Private/EdGraphSchema_K2.cpp
--- a/Engine/Source/Editor/BlueprintGraph/Private/EdGraphSchema_K2.cpp
+++ b/Engine/Source/Editor/BlueprintGraph/Private/EdGraphSchema_K2.cpp
@@ -105,6 +105,10 @@ const FName FBlueprintMetadata::MD_DynamicOutputParam(TEXT("DynamicOutputParam")
 const FName FBlueprintMetadata::MD_ArrayParam(TEXT("ArrayParm"));
 const FName FBlueprintMetadata::MD_ArrayDependentParam(TEXT("ArrayTypeDependentParams"));
 
+const FName FBlueprintMetadata::MD_Bitmask(TEXT("Bitmask"));
+const FName FBlueprintMetadata::MD_BitmaskEnum(TEXT("BitmaskEnum"));
+const FName FBlueprintMetadata::MD_Bitflags(TEXT("Bitflags"));
+
 //////////////////////////////////////////////////////////////////////////
 
 #define LOCTEXT_NAMESPACE "KismetSchema"
@@ -553,6 +557,7 @@ const FString UEdGraphSchema_K2::PC_Asset(TEXT("asset"));
 const FString UEdGraphSchema_K2::PC_AssetClass(TEXT("assetclass"));
 const FString UEdGraphSchema_K2::PSC_Self(TEXT("self"));
 const FString UEdGraphSchema_K2::PSC_Index(TEXT("index"));
+const FString UEdGraphSchema_K2::PSC_Bitmask(TEXT("bitmask"));
 const FString UEdGraphSchema_K2::PN_Execute(TEXT("execute"));
 const FString UEdGraphSchema_K2::PN_Then(TEXT("then"));
 const FString UEdGraphSchema_K2::PN_Completed(TEXT("Completed"));
@@ -2210,7 +2215,13 @@ private:
 	static FString GenerateTypeData(const FEdGraphPinType& PinType)
 	{
 		auto Obj = PinType.PinSubCategoryObject.Get();
-		return FString::Printf(TEXT("%s;%s;%s"), *PinType.PinCategory, *PinType.PinSubCategory, Obj ? *Obj->GetPathName() : TEXT(""));
+		FString PinSubCategory = PinType.PinSubCategory;
+		if (PinSubCategory.StartsWith(UEdGraphSchema_K2::PSC_Bitmask))
+		{
+			// Exclude the bitmask subcategory string from integral types so that autocast will work.
+			PinSubCategory = TEXT("");
+		}
+		return FString::Printf(TEXT("%s;%s;%s"), *PinType.PinCategory, *PinSubCategory, Obj ? *Obj->GetPathName() : TEXT(""));
 	}
 
 	static FString GenerateCastData(const FEdGraphPinType& InputPinType, const FEdGraphPinType& OutputPinType)
@@ -3152,11 +3163,24 @@ bool UEdGraphSchema_K2::ConvertPropertyToPinType(const UProperty* Property, /*ou
 	else if (Cast<const UIntProperty>(TestProperty) != NULL)
 	{
 		TypeOut.PinCategory = PC_Int;
+
+		if (TestProperty->HasMetaData(FBlueprintMetadata::MD_Bitmask))
+		{
+			TypeOut.PinSubCategory = PSC_Bitmask;
+		}
 	}
 	else if (const UByteProperty* ByteProperty = Cast<const UByteProperty>(TestProperty))
 	{
 		TypeOut.PinCategory = PC_Byte;
-		TypeOut.PinSubCategoryObject = ByteProperty->Enum;
+
+		if (TestProperty->HasMetaData(FBlueprintMetadata::MD_Bitmask))
+		{
+			TypeOut.PinSubCategory = PSC_Bitmask;
+		}
+		else
+		{
+			TypeOut.PinSubCategoryObject = ByteProperty->Enum;
+		}
 	}
 	else if (Cast<const UNameProperty>(TestProperty) != NULL)
 	{
@@ -3190,6 +3214,16 @@ bool UEdGraphSchema_K2::ConvertPropertyToPinType(const UProperty* Property, /*ou
 		return false;
 	}
 
+	if (TypeOut.PinSubCategory == PSC_Bitmask)
+	{
+		FString BitmaskEnumName = TestProperty->GetMetaData(TEXT("BitmaskEnum"));
+		if(!BitmaskEnumName.IsEmpty())
+		{
+			// @TODO: Potentially replace this with a serialized UEnum reference on the UProperty (e.g. UByteProperty::Enum)
+			TypeOut.PinSubCategoryObject = FindObject<UEnum>(ANY_PACKAGE, *BitmaskEnumName);
+		}
+	}
+
 	return true;
 }
 
@@ -3314,9 +3348,9 @@ FText UEdGraphSchema_K2::TypeToText(const FEdGraphPinType& Type)
 {
 	FText PropertyText;
 
-	if (Type.PinSubCategoryObject != NULL)
+	const UEdGraphSchema_K2* Schema = GetDefault<UEdGraphSchema_K2>();
+	if (Type.PinSubCategory != Schema->PSC_Bitmask && Type.PinSubCategoryObject != NULL)
 	{
-		const UEdGraphSchema_K2* Schema = GetDefault<UEdGraphSchema_K2>();
 		if (Type.PinCategory == Schema->PC_Byte)
 		{
 			FFormatNamedArguments Args;
@@ -3903,6 +3937,11 @@ bool UEdGraphSchema_K2::ArePinTypesCompatible(const FEdGraphPinType& Output, con
 				return true;
 			}
 		}
+		else if (PC_Byte == Output.PinCategory || PC_Int == Output.PinCategory)
+		{
+			// Bitmask integral types are compatible with non-bitmask integral types (of the same word size).
+			return Output.PinSubCategory.StartsWith(PSC_Bitmask) || Input.PinSubCategory.StartsWith(PSC_Bitmask);
+		}
 		else if (PC_Delegate == Output.PinCategory || PC_MCDelegate == Output.PinCategory)
 		{
 			auto CanUseFunction = [](const UFunction* Func) -> bool
@@ -5425,7 +5464,7 @@ void UEdGraphSchema_K2::BackwardCompatibilityNodeConversion(UEdGraph* Graph, boo
 	}
 }
 
-UEdGraphNode* UEdGraphSchema_K2::CreateSubstituteNode(UEdGraphNode* Node, const UEdGraph* Graph, FObjectInstancingGraph* InstanceGraph) const
+UEdGraphNode* UEdGraphSchema_K2::CreateSubstituteNode(UEdGraphNode* Node, const UEdGraph* Graph, FObjectInstancingGraph* InstanceGraph, TArray<FName>& InOutExtraNames) const
 {
 	// If this is an event node, create a unique custom event node as a substitute
 	UK2Node_Event* EventNode = Cast<UK2Node_Event>(Node);
@@ -5448,7 +5487,7 @@ UEdGraphNode* UEdGraphSchema_K2::CreateSubstituteNode(UEdGraphNode* Node, const
 		if(Blueprint && Blueprint->SkeletonGeneratedClass)
 		{
 			// Gather all names in use by the Blueprint class
-			TArray<FName> ExistingNamesInUse;
+			TArray<FName> ExistingNamesInUse = InOutExtraNames;
 			FBlueprintEditorUtils::GetFunctionNameList(Blueprint, ExistingNamesInUse);
 			FBlueprintEditorUtils::GetClassVariableList(Blueprint, ExistingNamesInUse);
 
@@ -5476,58 +5515,91 @@ UEdGraphNode* UEdGraphSchema_K2::CreateSubstituteNode(UEdGraphNode* Node, const
 			FString FunctionName;
 			const UK2Node_ActorBoundEvent* ActorBoundEventNode = Cast<const UK2Node_ActorBoundEvent>(EventNode);
 			const UK2Node_ComponentBoundEvent* CompBoundEventNode = Cast<const UK2Node_ComponentBoundEvent>(EventNode);
-			if(ActorBoundEventNode)
+
+			const UEdGraphNode* PreExistingNode = nullptr;
+			
+			if (InstanceGraph)
+			{
+				// Use a generic name for the new custom event
+				FunctionName = TEXT("CustomEvent");
+			}
+			else
 			{
-				FString TargetName = TEXT("None");
-				if(ActorBoundEventNode->EventOwner)
+				// Create a name for the custom event based off the original function
+				if (ActorBoundEventNode)
 				{
-					TargetName = ActorBoundEventNode->EventOwner->GetActorLabel();
-				}
+					FString TargetName = TEXT("None");
+					if (ActorBoundEventNode->EventOwner)
+					{
+						TargetName = ActorBoundEventNode->EventOwner->GetActorLabel();
+					}
 
-				FunctionName = FString::Printf(TEXT("%s_%s"), *ActorBoundEventNode->DelegatePropertyName.ToString(), *TargetName);
+					FunctionName = FString::Printf(TEXT("%s_%s"), *ActorBoundEventNode->DelegatePropertyName.ToString(), *TargetName);
+					PreExistingNode = FKismetEditorUtilities::FindBoundEventForActor(ActorBoundEventNode->GetReferencedLevelActor(), ActorBoundEventNode->DelegatePropertyName);
+				}
+				else if (CompBoundEventNode)
+				{
+					FunctionName = FString::Printf(TEXT("%s_%s"), *CompBoundEventNode->DelegatePropertyName.ToString(), *CompBoundEventNode->ComponentPropertyName.ToString());
+					PreExistingNode = FKismetEditorUtilities::FindBoundEventForComponent(Blueprint, CompBoundEventNode->DelegatePropertyName, CompBoundEventNode->ComponentPropertyName);
+				}
+				else if (EventNode->CustomFunctionName != NAME_None)
+				{
+					FunctionName = EventNode->CustomFunctionName.ToString();
+				}
+				else if (EventNode->bOverrideFunction)
+				{
+					FunctionName = EventNode->EventReference.GetMemberName().ToString();
+				}
+				else
+				{
+					FunctionName = CustomEventNode->GetName().Replace(TEXT("K2Node_"), TEXT(""), ESearchCase::CaseSensitive);
+				}
 			}
-			else if(CompBoundEventNode)
+
+			// Ensure the name does not overlap with other names
+			CustomEventNode->CustomFunctionName = FName(*FunctionName, FNAME_Find);
+			if (CustomEventNode->CustomFunctionName != NAME_None
+				&& ExistingNamesInUse.Contains(CustomEventNode->CustomFunctionName))
 			{
-				FunctionName = FString::Printf(TEXT("%s_%s"), *CompBoundEventNode->DelegatePropertyName.ToString(), *CompBoundEventNode->ComponentPropertyName.ToString());
+				int32 i = 0;
+				FString TempFuncName;
+
+				do
+				{
+					TempFuncName = FString::Printf(TEXT("%s_%d"), *FunctionName, ++i);
+					CustomEventNode->CustomFunctionName = FName(*TempFuncName, FNAME_Find);
+				} while (CustomEventNode->CustomFunctionName != NAME_None
+					&& ExistingNamesInUse.Contains(CustomEventNode->CustomFunctionName));
+
+				FunctionName = TempFuncName;
 			}
-			else if(EventNode->CustomFunctionName != NAME_None)
+
+			if (ActorBoundEventNode)
 			{
-				FunctionName = EventNode->CustomFunctionName.ToString();
+				PreExistingNode = FKismetEditorUtilities::FindBoundEventForActor(ActorBoundEventNode->GetReferencedLevelActor(), ActorBoundEventNode->DelegatePropertyName);
 			}
-			else if(EventNode->bOverrideFunction)
+			else if (CompBoundEventNode)
 			{
-				FunctionName = EventNode->EventReference.GetMemberName().ToString();
+				PreExistingNode = FKismetEditorUtilities::FindBoundEventForComponent(Blueprint, CompBoundEventNode->DelegatePropertyName, CompBoundEventNode->ComponentPropertyName);
 			}
 			else
 			{
-				FunctionName = CustomEventNode->GetName().Replace(TEXT("K2Node_"), TEXT(""), ESearchCase::CaseSensitive);
-			}
-
-			// Ensure that the new event name doesn't already exist as a variable or function name
-			if(InstanceGraph)
-			{
-				FunctionName += TEXT("_Copy");
-				CustomEventNode->CustomFunctionName = FName(*FunctionName, FNAME_Find);
-				if(CustomEventNode->CustomFunctionName != NAME_None
-					&& ExistingNamesInUse.Contains(CustomEventNode->CustomFunctionName))
+				if (Cast<UK2Node_CustomEvent>(EventNode))
 				{
-					int32 i = 0;
-					FString TempFuncName;
-
-					do 
-					{
-						TempFuncName = FString::Printf(TEXT("%s_%d"), *FunctionName, ++i);
-						CustomEventNode->CustomFunctionName = FName(*TempFuncName, FNAME_Find);
-					}
-					while(CustomEventNode->CustomFunctionName != NAME_None
-						&& ExistingNamesInUse.Contains(CustomEventNode->CustomFunctionName));
+					PreExistingNode = FBlueprintEditorUtils::FindCustomEventNode(Blueprint, EventNode->CustomFunctionName);
+				}
+				else
+				{
+					check(EventNode->FindEventSignatureFunction() != nullptr);
+					UClass* ClassOwner = EventNode->FindEventSignatureFunction()->GetOwnerClass()->GetAuthoritativeClass();
 
-					FunctionName = TempFuncName;
+					PreExistingNode = FBlueprintEditorUtils::FindOverrideForFunction(Blueprint, ClassOwner, EventNode->FindEventSignatureFunction()->GetFName());
 				}
 			}
 
 			// Should be a unique name now, go ahead and assign it
 			CustomEventNode->CustomFunctionName = FName(*FunctionName);
+			InOutExtraNames.Add(CustomEventNode->CustomFunctionName);
 
 			// Copy the pins from the old node to the new one that's replacing it
 			CustomEventNode->Pins = EventNode->Pins;
@@ -5570,13 +5642,31 @@ UEdGraphNode* UEdGraphSchema_K2::CreateSubstituteNode(UEdGraphNode* Node, const
 				}
 			}
 
+			if (PreExistingNode)
+			{
+				if (!Blueprint->PreCompileLog.IsValid())
+				{
+					Blueprint->PreCompileLog = TSharedPtr<FCompilerResultsLog>(new FCompilerResultsLog(false));
+					Blueprint->PreCompileLog->bSilentMode = false;
+					Blueprint->PreCompileLog->bAnnotateMentionedNodes = false;
+				}
+
+				// Append a warning to the node and to the logs
+				CustomEventNode->bHasCompilerMessage = true;
+				CustomEventNode->ErrorType = EMessageSeverity::Warning;
+				FFormatNamedArguments Args;
+				Args.Add(TEXT("NodeName"), CustomEventNode->GetNodeTitle(ENodeTitleType::ListView));
+				Args.Add(TEXT("OriginalNodeName"), FText::FromString(PreExistingNode->GetName()));
+				CustomEventNode->ErrorMsg = FText::Format(LOCTEXT("ReverseUpgradeWarning", "Conflicted with {OriginalNodeName} and was replaced as a Custom Event!"), Args).ToString();
+				Blueprint->PreCompileLog->Warning(*LOCTEXT("ReverseUpgradeWarning_Log", "Pasted node @@  conflicted with @@ and was replaced as a Custom Event!").ToString(), CustomEventNode, PreExistingNode);
+			}
 			// Return the new custom event node that we just created as a substitute for the original event node
 			return CustomEventNode;
 		}
 	}
 
 	// Use the default logic in all other cases
-	return UEdGraphSchema::CreateSubstituteNode(Node, Graph, InstanceGraph);
+	return UEdGraphSchema::CreateSubstituteNode(Node, Graph, InstanceGraph, InOutExtraNames);
 }
 
 int32 UEdGraphSchema_K2::GetNodeSelectionCount(const UEdGraph* Graph) const
