Copying //UE4/Dev-Blueprints to //UE4/Dev-Main (Source: //UE4/Dev-Blueprints @ 2967759)

==========================
MAJOR FEATURES + CHANGES
==========================

Change 2919729 on 2016/03/23 by [REDACTED_PERSON]

	Support for macros in the profiler
	#UEBP-177 - Macro instance handling
	#Proj Kismet, BlueprintProfiler

	- Adds support for timings inside macro calls
	- Extends the tunnel graph support to handle multiple entry/exit sites

Change 2938064 on 2016/04/08 by [REDACTED_PERSON]

	[UE-17794] The "Delete Unused Variable" feature now considers the GetClassDefaults node as well.

	change summary:
	- added external linkage to UK2Node_GetClassDefaults::FindClassPin().
	- added an include for the K2Node_GetClassDefaults header file to BlueprintGraphDefinitions.h.
	- added UK2Node_GetClassDefaults::GetInputClass() as a public API w/ external linkage; moved default 'nullptr' param logic into this impl.
	- modified FBlueprintEditorUtils::IsVariableUsed() to add an extra check for a GetClassDefaults node with a visible output pin for the variable that's also connected.
	- modified UK2Node_GetClassDefaults::GetInputClass() to return the generated skeleton class for Blueprint class types.

Change 2938088 on 2016/04/08 by [REDACTED_PERSON]

	Making bytecode produced by latent action nodes deterministic.

Change 2938101 on 2016/04/08 by [REDACTED_PERSON]

	Fixing a bug where the compile summary was not being reported because another compiler log was lingering - making it so the MathExpression node compiler log is not initialized with intent to report its own summary (it gets folded into the primary log).

Change 2938121 on 2016/04/08 by [REDACTED_PERSON]

	Remove a few redundant MBASM calls on variable type change.

Change 2940295 on 2016/04/11 by [REDACTED_PERSON]

	We now 'tag subobjects' of a blueprint even if it's part of the rootset, this means we correctly detect references to the BPGC in FPendingDelete::CheckForReferences(). Original rootset check dates to 2012 and I can find no justification for it currently.

Change 2943227 on 2016/04/13 by [REDACTED_PERSON]

	Fixed improper detection of functions from interfaces that themselves inherit from some other interface

Change 2944270 on 2016/04/14 by [REDACTED_PERSON]

	[UEBP-176] First pass at BP graph node heat map visualization while profiling.

	change summary:
	- added an "indicator overlay" to graph node widget layouts
	- added a heat mode "mode" selector widget to the BP profiler view panel
	- extended IBlueprintProfilerInterface to include APIs for accessing current heat map mode state
	- added FScriptNodePerfData::GetBlueprintPerfDataForAllTracePaths() (tentative - may need revisiting)
	- added SGraphNode::GetNodeIndicatorOverlayColor() and GetNodeIndicatorOverlayVisibility() delegates
	- added BP-specific delegate overrides to SGraphNodeK2Base; extended to include both compact and variable nodes

Change 2946932 on 2016/04/18 by [REDACTED_PERSON]

	Guarding against invalid EdGraphPins (ones that have been moved to the transient package) when constructing the widget - prevents a crash that we've been unable to repro or determine the cause of (turns it instead into an ensure, so we can collect more contextual information on the issue).

Change 2949968 on 2016/04/20 by [REDACTED_PERSON]

	Array access out of bounds by value is a warning again, added ability to elevate individual warnings on a per project basis (or supress them)

Change 2950113 on 2016/04/20 by [REDACTED_PERSON]

	Removed GBlueprintCompileTime, it was not accurate. Printing BlueprintCompileAndLoadTimerData.GetTime() at start instead

Change 2951086 on 2016/04/21 by [REDACTED_PERSON]

	This change addresses the edge case in the blueprint profiler that caused stats to fail when tunnel nodes were linked through to other tunnel nodes.
	#Proj Kismet, BlueprintProfiler

Change 2951336 on 2016/04/21 by [REDACTED_PERSON]

	This change enables blueprint breakpoints during instrumented conditions.
	#Proj CoreUObject, BlueprintProfiler, UnrealEd, KismetCompiler

Change 2951406 on 2016/04/21 by [REDACTED_PERSON]

	Fix for blueprint profiler stats for the top level blueprint stat entry not updating correctly.
	#Proj Kismet

Change 2951832 on 2016/04/21 by [REDACTED_PERSON]

	Fix for certain blueprint profiler stats not being updated and collected at the blueprint container level due to incorrect tracepaths.
	#Proj Kismet

	#info This should fix the node heatmaps as a side effect.

Change 2956696 on 2016/04/26 by [REDACTED_PERSON]

	Tweak fix for macros being BS_Dirty after loading. The current fix had the side effect of not recompiling clients of the macro after making a change to the macro and entering PIE

Change 2957564 on 2016/04/27 by [REDACTED_PERSON]

	Various fixes related to nativized enums.

Change 2961626 on 2016/04/29 by [REDACTED_PERSON]

	Merging //UE4/Dev-Main to Dev-Blueprints (//UE4/Dev-Blueprints)

Change 2962747 on 2016/05/02 by [REDACTED_PERSON]

	GameplayAbility handles BPGC and DynamicClass the same way.
	C++ backend do not assumes that some literal objects are UserDefinesEnum/UserDefinedStruct/BlueprintGeneratedClass.

Change 2965679 on 2016/05/04 by [REDACTED_PERSON]

	Increased stack size (384k) of threads spawned by [REDACTED_PERSON].
	(Temporary?) fix for stack overflow, when amination in Orion are evaluated.

Change 2965758 on 2016/05/04 by [REDACTED_PERSON]

	Fixed CDO creation in async loaded Dynamic Class.
	Fixed too restrict cast assertions.

Change 2966543 on 2016/05/04 by [REDACTED_PERSON]

Change 2966839 on 2016/05/04 by [REDACTED_PERSON]

	Typo IMPLEMENT_MODULE creates weird linking error, also may need entry in Target.cs to get BlueprintRuntime to build. Copying that pattern from ___LoadingScreen modules

Change 2967347 on 2016/05/05 by [REDACTED_PERSON]

	CommandUtils.GetDirectoryName should not be used with directory path (but only with file path), because it cannot handle paths like "e:\My11Project 4.13" (containing '.'). It seems useless with directory path anyway.

diff --git a/Engine/Source/Developer/BlueprintCompilerCppBackend/Private/BlueprintCompilerCppBackend.cpp b/Engine/Source/Developer/BlueprintCompilerCppBackend/Private/BlueprintCompilerCppBackend.cpp
--- a/Engine/Source/Developer/BlueprintCompilerCppBackend/Private/BlueprintCompilerCppBackend.cpp
+++ b/Engine/Source/Developer/BlueprintCompilerCppBackend/Private/BlueprintCompilerCppBackend.cpp
@@ -4,6 +4,7 @@
 #include "BlueprintCompilerCppBackend.h"
 #include "BlueprintCompilerCppBackendUtils.h"
 #include "EdGraphSchema_K2.h"
+#include "Kismet/KismetNodeHelperLibrary.h"
 
 // Generates single "if" scope. Its condition checks context of given term.
 struct FSafeContextScopedEmmitter
@@ -521,65 +522,90 @@ FString FBlueprintCompilerCppBackend::EmitCallStatmentInner(FEmitterLocalContext
 		}
 	}
 
-	// Emit object to call the method on
-	if (bInterfaceCallExecute)
+	auto HandleSpeciallyNativizedFunction = [&]() -> bool
 	{
-		auto ContextInterfaceClass = CastChecked<UClass>(Statement.FunctionContext->Type.PinSubCategoryObject.Get());
-		ensure(ContextInterfaceClass->IsChildOf<UInterface>());
-		Result += FString::Printf(TEXT("%s::Execute_%s(%s.GetObject() ")
-			, *FEmitHelper::GetCppName(ContextInterfaceClass)
-			, *FunctionToCallOriginalName
-			, *TermToText(EmitterContext, Statement.FunctionContext, false));
-	}
-	else
-	{
-		auto FunctionOwner = Statement.FunctionToCall->GetOwnerClass();
-		auto OwnerBPGC = Cast<UBlueprintGeneratedClass>(FunctionOwner);
-		const bool bUnconvertedClass = OwnerBPGC && !EmitterContext.Dependencies.WillClassBeConverted(OwnerBPGC);
-		if (bUnconvertedClass)
+		const UFunction* const EnumeratorUserFriendlyName = UKismetNodeHelperLibrary::StaticClass()->FindFunctionByName(GET_FUNCTION_NAME_CHECKED(UKismetNodeHelperLibrary, GetEnumeratorUserFriendlyName));
+		if (EnumeratorUserFriendlyName == Statement.FunctionToCall)
 		{
-			ensure(!Statement.bIsParentContext); //unsupported yet
-			ensure(!bStaticCall); //unsupported yet
-			ensure(bCallOnDifferentObject); //unexpected
-			const FString WrapperName = FString::Printf(TEXT("FUnconvertedWrapper__%s"), *FEmitHelper::GetCppName(OwnerBPGC));
-			const FString CalledObject = bCallOnDifferentObject ? TermToText(EmitterContext, Statement.FunctionContext, false) : TEXT("this");
-			Result += FString::Printf(TEXT("%s(%s)."), *WrapperName, *CalledObject);
+			const FBPTerminal* const EnumTerminal = Statement.RHS.IsValidIndex(0) ? Statement.RHS[0] : nullptr;
+			// We need the special path only for converted enums. All UUserDefinedEnum are converted.
+			const UUserDefinedEnum* const Enum = EnumTerminal ? Cast<UUserDefinedEnum>(EnumTerminal->ObjectLiteral) : nullptr;
+			const FBPTerminal* const ValueTerminal = Statement.RHS.IsValidIndex(1) ? Statement.RHS[1] : nullptr;
+			if (Enum && ValueTerminal)
+			{
+				const FString Value = TermToText(EmitterContext, ValueTerminal);
+				const FString EnumName = FEmitHelper::GetCppName(Enum);
+				Result += FString::Printf(TEXT("%s__GetUserFriendlyName(EnumToByte<%s>(%s))"), *EnumName, *EnumName, *Value);
+				return true;
+			}
 		}
-		else if (bStaticCall)
+		return false;
+	};
+
+	const bool bIsSpecialCase = HandleSpeciallyNativizedFunction();
+	if (!bIsSpecialCase)
+	{
+		// Emit object to call the method on
+		if (bInterfaceCallExecute)
 		{
-			const bool bIsCustomThunk = Statement.FunctionToCall->GetBoolMetaData(TEXT("CustomThunk")) 
-				|| Statement.FunctionToCall->HasMetaData(TEXT("CustomStructureParam")) 
-				|| Statement.FunctionToCall->HasMetaData(TEXT("ArrayParm"));
-			auto OwnerClass = Statement.FunctionToCall->GetOuterUClass();
-			Result += bIsCustomThunk ? TEXT("FCustomThunkTemplates::") : FString::Printf(TEXT("%s::"), *FEmitHelper::GetCppName(OwnerClass));
+			auto ContextInterfaceClass = CastChecked<UClass>(Statement.FunctionContext->Type.PinSubCategoryObject.Get());
+			ensure(ContextInterfaceClass->IsChildOf<UInterface>());
+			Result += FString::Printf(TEXT("%s::Execute_%s(%s.GetObject() ")
+				, *FEmitHelper::GetCppName(ContextInterfaceClass)
+				, *FunctionToCallOriginalName
+				, *TermToText(EmitterContext, Statement.FunctionContext, false));
 		}
-		else if (bCallOnDifferentObject) //@TODO: Badness, could be a self reference wired to another instance!
+		else
 		{
-			Result += FString::Printf(TEXT("%s->"), *TermToText(EmitterContext, Statement.FunctionContext, false));
-		}
+			auto FunctionOwner = Statement.FunctionToCall->GetOwnerClass();
+			auto OwnerBPGC = Cast<UBlueprintGeneratedClass>(FunctionOwner);
+			const bool bUnconvertedClass = OwnerBPGC && !EmitterContext.Dependencies.WillClassBeConverted(OwnerBPGC);
+			if (bUnconvertedClass)
+			{
+				ensure(!Statement.bIsParentContext); //unsupported yet
+				ensure(!bStaticCall); //unsupported yet
+				ensure(bCallOnDifferentObject); //unexpected
+				const FString WrapperName = FString::Printf(TEXT("FUnconvertedWrapper__%s"), *FEmitHelper::GetCppName(OwnerBPGC));
+				const FString CalledObject = bCallOnDifferentObject ? TermToText(EmitterContext, Statement.FunctionContext, false) : TEXT("this");
+				Result += FString::Printf(TEXT("%s(%s)."), *WrapperName, *CalledObject);
+			}
+			else if (bStaticCall)
+			{
+				const bool bIsCustomThunk = Statement.FunctionToCall->GetBoolMetaData(TEXT("CustomThunk"))
+					|| Statement.FunctionToCall->HasMetaData(TEXT("CustomStructureParam"))
+					|| Statement.FunctionToCall->HasMetaData(TEXT("ArrayParm"));
+				auto OwnerClass = Statement.FunctionToCall->GetOuterUClass();
+				Result += bIsCustomThunk ? TEXT("FCustomThunkTemplates::") : FString::Printf(TEXT("%s::"), *FEmitHelper::GetCppName(OwnerClass));
+			}
+			else if (bCallOnDifferentObject) //@TODO: Badness, could be a self reference wired to another instance!
+			{
+				Result += FString::Printf(TEXT("%s->"), *TermToText(EmitterContext, Statement.FunctionContext, false));
+			}
 
-		if (Statement.bIsParentContext)
-		{
-			Result += TEXT("Super::");
-		}
-		Result += FunctionToCallOriginalName;
+			if (Statement.bIsParentContext)
+			{
+				Result += TEXT("Super::");
+			}
+			Result += FunctionToCallOriginalName;
 
-		if (Statement.bIsParentContext && FEmitHelper::ShouldHandleAsNativeEvent(Statement.FunctionToCall, false))
+			if (Statement.bIsParentContext && FEmitHelper::ShouldHandleAsNativeEvent(Statement.FunctionToCall, false))
+			{
+				ensure(!bCallOnDifferentObject);
+				Result += TEXT("_Implementation");
+			}
+
+			// Emit method parameter list
+			Result += TEXT("(");
+		}
+		const FString ParameterList = EmitMethodInputParameterList(EmitterContext, Statement);
+		if (bInterfaceCallExecute && !ParameterList.IsEmpty())
 		{
-			ensure(!bCallOnDifferentObject);
-			Result += TEXT("_Implementation");
+			Result += TEXT(", ");
 		}
-
-		// Emit method parameter list
-		Result += TEXT("(");
+		Result += ParameterList;
+		Result += TEXT(")");
 	}
-	const FString ParameterList = EmitMethodInputParameterList(EmitterContext, Statement);
-	if (bInterfaceCallExecute && !ParameterList.IsEmpty())
-	{
-		Result += TEXT(", ");
-	}
-	Result += ParameterList;
-	Result += TEXT(")");
+
 	Result += CloseCast;
 	if (!bInline)
 	{
diff --git a/Engine/Source/Developer/BlueprintCompilerCppBackend/Private/BlueprintCompilerCppBackendBase.cpp b/Engine/Source/Developer/BlueprintCompilerCppBackend/Private/BlueprintCompilerCppBackendBase.cpp
--- a/Engine/Source/Developer/BlueprintCompilerCppBackend/Private/BlueprintCompilerCppBackendBase.cpp
+++ b/Engine/Source/Developer/BlueprintCompilerCppBackend/Private/BlueprintCompilerCppBackendBase.cpp
@@ -125,7 +125,7 @@ void FBlueprintCompilerCppBackendBase::DeclareDelegates(FEmitterLocalContext& Em
 		for (int32 I = 0; I < Delegates.Num(); ++I)
 		{
 			UFunction* TargetFn = Delegates[I]->SignatureFunction;
-			for (int32 J = I + 1; J < Delegates.Num(); ++J)
+			for (int32 J = I + 1; J < Delegates.Num();)
 			{
 				if (TargetFn == Delegates[J]->SignatureFunction)
 				{
@@ -133,6 +133,10 @@ void FBlueprintCompilerCppBackendBase::DeclareDelegates(FEmitterLocalContext& Em
 					Delegates[J] = Delegates[Delegates.Num() - 1];
 					Delegates.RemoveAt(Delegates.Num() - 1);
 				}
+				else
+				{
+					J++;
+				}
 			}
 		}
 
@@ -680,27 +684,30 @@ FString FBlueprintCompilerCppBackendBase::GenerateCodeFromEnum(UUserDefinedEnum*
 	}
 	Header.DecreaseIndent();
 	Header.AddLine(TEXT("};"));
-	/*
-	Header.AddLine(FString::Printf(TEXT("FString %s__GetUserFriendlyName(int32 InValue)"), *EnumCppName));
+
+	Header.AddLine(FString::Printf(TEXT("inline FString %s__GetUserFriendlyName(int32 InValue)"), *EnumCppName));
 	Header.AddLine(TEXT("{"));
 	Header.IncreaseIndent();
 
-	Header.AddLine(TEXT("switch(InValue)"));
+	Header.AddLine(TEXT("FText Text;"));
+	Header.AddLine(FString::Printf(TEXT("const auto EnumValue = static_cast<%s>(InValue);"), *EnumCppName));
+	Header.AddLine(TEXT("switch(EnumValue)"));
 	Header.AddLine(TEXT("{"));
 	Header.IncreaseIndent();
 	for (int32 Index = 0; Index < SourceEnum->NumEnums(); ++Index)
 	{
-		const FString DisplayNameStr = SourceEnum->GetEnumText(Index).ToString().ReplaceCharWithEscapedChar();
-		Header.AddLine(FString::Printf(TEXT("case %s::%s: return FString(TEXT(\"%s\"));"), *EnumCppName, *SourceEnum->GetEnumName(Index), *DisplayNameStr));
+		FString DisplayNameStr;
+		FTextStringHelper::WriteToString(DisplayNameStr, SourceEnum->GetEnumText(Index));
+		Header.AddLine(FString::Printf(TEXT("case %s::%s: FTextStringHelper::ReadFromString(TEXT(\"%s\"), Text); break;"), *EnumCppName, *SourceEnum->GetEnumName(Index), *DisplayNameStr.ReplaceCharWithEscapedChar()));
 	}
+	Header.AddLine(TEXT("default: ensure(false);"));
 	Header.DecreaseIndent();
 	Header.AddLine(TEXT("};"));
 
-	Header.AddLine(TEXT("ensure(false);"));
-	Header.AddLine(TEXT("return FString();"));
+	Header.AddLine(TEXT("return Text.ToString();"));
 	Header.DecreaseIndent();
 	Header.AddLine(TEXT("};"));
-	*/
+
 	return Header.Result;
 }
 
diff --git a/Engine/Source/Developer/BlueprintCompilerCppBackend/Private/BlueprintCompilerCppBackendUtils.cpp b/Engine/Source/Developer/BlueprintCompilerCppBackend/Private/BlueprintCompilerCppBackendUtils.cpp
--- a/Engine/Source/Developer/BlueprintCompilerCppBackend/Private/BlueprintCompilerCppBackendUtils.cpp
+++ b/Engine/Source/Developer/BlueprintCompilerCppBackend/Private/BlueprintCompilerCppBackendUtils.cpp
@@ -35,7 +35,10 @@ FString FEmitterLocalContext::FindGloballyMappedObject(const UObject* Object, co
 	auto ClassString = [&]() -> FString
 	{
 		const UClass* ObjectClassToUse = ExpectedClass ? ExpectedClass : GetFirstNativeOrConvertedClass(Object->GetClass());
-		return FEmitHelper::GetCppName(ObjectClassToUse);
+		ObjectClassToUse = (UUserDefinedEnum::StaticClass() == ObjectClassToUse) ? UEnum::StaticClass() : ObjectClassToUse;
+		ObjectClassToUse = (UUserDefinedStruct::StaticClass() == ObjectClassToUse) ? UScriptStruct::StaticClass() : ObjectClassToUse;
+		ObjectClassToUse = (!ExpectedClass && ObjectClassToUse && ObjectClassToUse->IsChildOf<UBlueprintGeneratedClass>()) ? UClass::StaticClass() : ObjectClassToUse;
+		return FEmitHelper::GetCppName((UUserDefinedEnum::StaticClass() == ObjectClassToUse) ? UEnum::StaticClass() : ObjectClassToUse);
 	};
 	
 	if (ActualClass && Object 
@@ -139,14 +142,6 @@ FString FEmitterLocalContext::FindGloballyMappedObject(const UObject* Object, co
 		return FString::Printf(TEXT("%s::StaticStruct()"), *FEmitHelper::GetCppName(UDS));
 	}
 
-	if (auto UDE = Cast<UEnum>(Object))
-	{
-		// TODO:
-		// TODO: check if supported 
-		//return FString::Printf(TEXT("%s_StaticEnum()"), *UDE->GetName());
-		return FString::Printf(TEXT("FindObjectChecked<UEnum>(ANY_PACKAGE, TEXT(\"%s\"))"), *FEmitHelper::GetCppName(UDE));
-	}
-
 	// TODO: handle subobjects
 	ensure(!bLoadIfNotFound || Object);
 	if (Object && (bLoadIfNotFound || bTryUsedAssetsList))
@@ -979,12 +974,12 @@ FString FEmitHelper::LiteralTerm(FEmitterLocalContext& EmitterContext, const FEd
 	{
 		if (auto FoundClass = Cast<const UClass>(LiteralObject))
 		{
-			const FString MappedObject = EmitterContext.FindGloballyMappedObject(LiteralObject);
+			const FString MappedObject = EmitterContext.FindGloballyMappedObject(LiteralObject, UClass::StaticClass());
 			if (!MappedObject.IsEmpty())
 			{
 				return MappedObject;
 			}
-			return FString::Printf(TEXT("LoadClass<UObject>(nullptr, TEXT(\"%s\"), nullptr, 0, nullptr)"), *(LiteralObject->GetPathName().ReplaceCharWithEscapedChar()));
+			return FString::Printf(TEXT("LoadClass<UClass>(nullptr, TEXT(\"%s\"), nullptr, 0, nullptr)"), *(LiteralObject->GetPathName().ReplaceCharWithEscapedChar()));
 		}
 		return FString(TEXT("nullptr"));
 	}
@@ -1240,7 +1235,7 @@ bool FEmitHelper::GenerateAutomaticCast(FEmitterLocalContext& EmitterContext, co
 		if (!RType.bIsArray && LClass && RClass && LClass->IsChildOf(RClass) && !RClass->IsChildOf(LClass))
 		{
 			OutCastBegin = FString::Printf(TEXT("CastChecked<%s>("), *FEmitHelper::GetCppName(LClass));
-			OutCastEnd = TEXT(")");
+			OutCastEnd = TEXT(", ECastCheckedType::NullAllowed)");
 			return true;
 		}
 		else if (RType.bIsArray && LClass && RClass && (LClass->IsChildOf(RClass) || RClass->IsChildOf(LClass)) && (LClass != RClass))
diff --git a/Engine/Source/Developer/BlueprintCompilerCppBackend/Private/BlueprintCompilerCppBackendValueHelper.cpp b/Engine/Source/Developer/BlueprintCompilerCppBackend/Private/BlueprintCompilerCppBackendValueHelper.cpp
--- a/Engine/Source/Developer/BlueprintCompilerCppBackend/Private/BlueprintCompilerCppBackendValueHelper.cpp
+++ b/Engine/Source/Developer/BlueprintCompilerCppBackend/Private/BlueprintCompilerCppBackendValueHelper.cpp
@@ -201,7 +201,7 @@ void FEmitDefaultValueHelper::InnerGenerate(FEmitterLocalContext& Context, const
 
 			if (RegularInnerStruct)
 			{
-				const FString InnerStructStr = Context.FindGloballyMappedObject(RegularInnerStruct);
+				const FString InnerStructStr = Context.FindGloballyMappedObject(RegularInnerStruct, UScriptStruct::StaticClass());
 				Context.AddLine(FString::Printf(TEXT("%s->InitializeStruct(%s.GetData(), %d);"), *InnerStructStr, *PathToMember, ScriptArrayHelper.Num()));
 			}
 		}
@@ -653,7 +653,7 @@ public:
 		{
 			return Result;
 		}
-		
+
 		for (UObject* Outer = Item; Outer; Outer = Outer->GetOuter())
 		{
 			if (!Result.IsEmpty())
@@ -811,6 +811,15 @@ void FEmitDefaultValueHelper::GenerateCustomDynamicClassInitialization(FEmitterL
 		Context.AddLine(FString::Printf(TEXT("InDynamicClass->ReferencedConvertedFields.Add(%s);"), *StructConstructor));
 	}
 
+	if (Context.Dependencies.ConvertedEnum.Num())
+	{
+		Context.AddLine(TEXT("// List of all referenced converted enums"));
+	}
+	for (auto LocEnum : Context.Dependencies.ConvertedEnum)
+	{
+		Context.AddLine(FString::Printf(TEXT("InDynamicClass->ReferencedConvertedFields.Add(LoadObject<UEnum>(nullptr, TEXT(\"%s\")));"), *(LocEnum->GetPathName().ReplaceCharWithEscapedChar())));
+	}
+
 	TArray<UActorComponent*> ActorComponentTempatesOwnedByClass = BPGC->ComponentTemplates;
 	// Gather all CT from SCS and IH, the remaining ones are generated for class..
 	if (auto SCS = BPGC->SimpleConstructionScript)
diff --git a/Engine/Source/Developer/BlueprintProfiler/Private/BlueprintProfiler.cpp b/Engine/Source/Developer/BlueprintProfiler/Private/BlueprintProfiler.cpp
--- a/Engine/Source/Developer/BlueprintProfiler/Private/BlueprintProfiler.cpp
+++ b/Engine/Source/Developer/BlueprintProfiler/Private/BlueprintProfiler.cpp
@@ -18,6 +18,7 @@ FBlueprintProfiler::FBlueprintProfiler()
 	: bProfilerActive(false)
 	, bProfilingCaptureActive(false)
 	, bPIEActive(false)
+	, GraphNodeHeatMapDisplayMode(EBlueprintProfilerHeatMapDisplayMode::None)
 {
 }
 
diff --git a/Engine/Source/Developer/BlueprintProfiler/Private/ScriptInstrumentationPlayback.cpp b/Engine/Source/Developer/BlueprintProfiler/Private/ScriptInstrumentationPlayback.cpp
--- a/Engine/Source/Developer/BlueprintProfiler/Private/ScriptInstrumentationPlayback.cpp
+++ b/Engine/Source/Developer/BlueprintProfiler/Private/ScriptInstrumentationPlayback.cpp
@@ -380,7 +380,7 @@ bool FBlueprintFunctionContext::DetectCyclicLinks(TSharedPtr<FScriptExecutionNod
 				{
 					// Break links and flag cycle linkage.
 					FScriptExecNodeParams CycleLinkParams;
-					const FString NodeName = FString::Printf(TEXT("CyclicLinkTo_%i_%s"), LinkedNode.Key, *LinkedExecNode->GetName().ToString());
+					const FString NodeName = FString::Printf(TEXT("CyclicLinkTo_%i_%s"), ExecutionNodes.Num(), *LinkedExecNode->GetName().ToString());
 					CycleLinkParams.NodeName = FName(*NodeName);
 					CycleLinkParams.ObservedObject = LinkedExecNode->GetObservedObject();
 					CycleLinkParams.DisplayName = LinkedExecNode->GetDisplayName();
@@ -653,6 +653,7 @@ void FBlueprintFunctionContext::MapExecPins(TSharedPtr<FScriptExecutionNode> Exe
 {
 	const bool bBranchedExecution = ExecNode->IsBranch();
 	int32 NumUnwiredPins = 0;
+	const UEdGraphPin* FinalExecPin = nullptr;
 	for (auto Pin : Pins)
 	{
 		int32 PinScriptCodeOffset = GetCodeLocationFromPin(Pin);
@@ -708,11 +709,19 @@ void FBlueprintFunctionContext::MapExecPins(TSharedPtr<FScriptExecutionNode> Exe
 				// Check for re-entrant exit site
 				if (ExecNode->HasFlags(EScriptExecutionNodeFlags::SequentialBranch))
 				{
-					// Last pin in sequence is not re-entrant.
-					if (Pin != Pins.Last())
+					// Locate the last linked exec pin
+					if (!FinalExecPin)
 					{
-						TunnelExitNode->AddFlags(EScriptExecutionNodeFlags::ReEntrantTunnelPin);
+						for (auto SeqPin : Pins)
+						{
+							if (SeqPin->LinkedTo.Num())
+							{
+								FinalExecPin = SeqPin;
+							}
+						}
 					}
+					// Last pin in sequence is not re-entrant, but all others are assumed to be.
+					TunnelExitNode->AddFlags(Pin == FinalExecPin ? EScriptExecutionNodeFlags::TunnelFinalExitPin : EScriptExecutionNodeFlags::ReEntrantTunnelPin);
 				}
 				if (bBranchedExecution)
 				{
@@ -750,10 +759,26 @@ void FBlueprintFunctionContext::MapExecPins(TSharedPtr<FScriptExecutionNode> Exe
 TSharedPtr<FScriptExecutionNode> FBlueprintFunctionContext::MapTunnelEntry(const FName TunnelName, const UEdGraphPin* ExecPin, const UEdGraphPin* TunnelPin, const int32 PinScriptOffset)
 {
 	// Retrieve tunnel function context
-	const UEdGraphNode* TunnelNode = TunnelPin->GetOwningNode();
+	const UK2Node_Tunnel* TunnelNode = Cast<UK2Node_Tunnel>(TunnelPin->GetOwningNode());
 	TSharedPtr<FBlueprintFunctionContext> TunnelContext = GetTunnelContextFromNode(TunnelNode);
 	const int32 TunnelId = TunnelContext->GetTunnelIdFromName(TunnelName);
 	check(TunnelId != INDEX_NONE);
+	// Create node for collapsed tunnel/macro
+	TSharedPtr<FScriptExecutionTunnelInstance> TunnelInstance = GetTypedProfilerDataForNode<FScriptExecutionTunnelInstance>(TunnelNode->GetFName());
+	if (!TunnelInstance.IsValid())
+	{
+		FScriptExecNodeParams TunnelInstanceParams;
+		TunnelInstanceParams.NodeName = TunnelNode->GetFName();
+		TunnelInstanceParams.ObservedObject = TunnelNode;
+		TunnelInstanceParams.DisplayName = TunnelNode->GetNodeTitle(ENodeTitleType::ListView);
+		TunnelInstanceParams.Tooltip = LOCTEXT("NavigateToNodeLocationHyperlink_ToolTip", "Navigate to the Node");
+		TunnelInstanceParams.NodeFlags = EScriptExecutionNodeFlags::FunctionTunnel;
+		TunnelInstanceParams.IconColor = FLinearColor(1.f, 1.f, 1.f, 0.8f);
+		const FSlateBrush* TunnelIcon = TunnelNode->ShowPaletteIconOnNode() ?	FEditorStyle::GetBrush(TunnelNode->GetPaletteIcon(TunnelInstanceParams.IconColor)) :
+																				FEditorStyle::GetBrush(TEXT("BlueprintProfiler.BPNode"));
+		TunnelInstanceParams.Icon = const_cast<FSlateBrush*>(TunnelIcon);
+		TunnelInstance = CreateTypedExecutionNode<FScriptExecutionTunnelInstance>(TunnelInstanceParams);
+	}
 	// Create tunnel graph instance node
 	FScriptExecNodeParams TunnelParams;
 	TunnelParams.NodeName = FName(*FString::Printf(TEXT("%s_EntryPointId%i"), *TunnelNode->GetFName().ToString(), TunnelId));
@@ -765,7 +790,9 @@ TSharedPtr<FScriptExecutionNode> FBlueprintFunctionContext::MapTunnelEntry(const
 	const FSlateBrush* TunnelIcon = TunnelNode->ShowPaletteIconOnNode() ?	FEditorStyle::GetBrush(TunnelNode->GetPaletteIcon(TunnelParams.IconColor)) :
 																			FEditorStyle::GetBrush(TEXT("BlueprintProfiler.BPNode"));
 	TunnelParams.Icon = const_cast<FSlateBrush*>(TunnelIcon);
-	TSharedPtr<FScriptExecutionNode> TunnelCallNode = CreateExecutionNode(TunnelParams);
+	TSharedPtr<FScriptExecutionTunnelEntry> TunnelCallNode = MakeShareable(new FScriptExecutionTunnelEntry(TunnelParams, TunnelNode));
+	ExecutionNodes.FindOrAdd(TunnelParams.NodeName) = TunnelCallNode;
+	TunnelInstance->AddCustomEntryPoint(TunnelCallNode);
 	// Add entry point children
 	TSharedPtr<FScriptExecutionNode> EntryPinNode = TunnelContext->GetProfilerDataForNode(TunnelName);
 	check (EntryPinNode.IsValid());
@@ -779,35 +806,60 @@ TSharedPtr<FScriptExecutionNode> FBlueprintFunctionContext::MapTunnelEntry(const
 	// Map exit points
 	TArray<TSharedPtr<FScriptExecutionNode>> TunnelExitPoints;
 	TunnelContext->GetTunnelExitPoints(TunnelName, TunnelExitPoints);
-	if (const UK2Node_Tunnel* TunnelGraphNode = Cast<UK2Node_Tunnel>(TunnelNode))
+	TMap<FName, const UEdGraphPin*> NameToPinMap;
+	for (auto ExitPin : TunnelNode->Pins)
+	{
+		if (ExitPin->PinType.PinCategory == UEdGraphSchema_K2::PC_Exec && ExitPin->LinkedTo.Num())
+		{
+			NameToPinMap.Add(FName(*ExitPin->PinName)) = ExitPin;
+		}
+	}
+	// Register tunnel input pin node using the script offset from the first script offset emitted pin.
+	if (const UEdGraphPin** InputPin = NameToPinMap.Find(TunnelName))
 	{
-		TMap<FName, const UEdGraphPin*> NameToPinMap;
-		for (auto ExitPin : TunnelGraphNode->Pins)
+		if ((*InputPin)->LinkedTo.Num())
 		{
-			if (ExitPin->Direction == EGPD_Output && ExitPin->PinType.PinCategory == UEdGraphSchema_K2::PC_Exec && ExitPin->LinkedTo.Num())
+			if (const UEdGraphPin* InputLinkedPin = (*InputPin)->LinkedTo[0])
 			{
-				NameToPinMap.Add(FName(*ExitPin->PinName)) = ExitPin;
+				const int32 EntryScriptCodeOffset = GetCodeLocationFromPin(InputLinkedPin);
+				ScriptOffsetToPins.Add(EntryScriptCodeOffset) = *InputPin;
+				TunnelEntrySites.Add(EntryScriptCodeOffset) = TunnelCallNode;
 			}
 		}
-		for (auto CurrExitPoint : TunnelExitPoints)
+	}
+	for (auto CurrExitPoint : TunnelExitPoints)
+	{
+		if (const UEdGraphPin** ExitPin = NameToPinMap.Find(CurrExitPoint->GetName()))
 		{
-			if (const UEdGraphPin** ExitPin = NameToPinMap.Find(CurrExitPoint->GetName()))
+			for (auto LinkedExitPin : (*ExitPin)->LinkedTo)
 			{
-				for (auto LinkedExitPin : (*ExitPin)->LinkedTo)
+				const int32 ExitScriptCodeOffset = GetCodeLocationFromPin(*ExitPin);
+				// Add entry site for exit script offset
+				TunnelExitSites.Add(ExitScriptCodeOffset) = TunnelCallNode;
+				// map tunnel exit links
+				UEdGraphNode* LinkedExitNode = LinkedExitPin->GetOwningNode();
+				TSharedPtr<FScriptExecutionNode> LinkedExecNode;
+				if (LinkedExitNode->IsA<UK2Node_Tunnel>())
+				{
+					const FName LinkedTunnelName(*LinkedExitPin->PinName);
+					LinkedExecNode = MapTunnelEntry(LinkedTunnelName, *ExitPin, LinkedExitPin, ExitScriptCodeOffset);
+				}
+				else
 				{
-					TSharedPtr<FScriptExecutionNode> LinkedExecNode = MapNodeExecution(LinkedExitPin->GetOwningNode());
-					if (LinkedExecNode.IsValid())
+					LinkedExecNode = MapNodeExecution(LinkedExitPin->GetOwningNode());
+				}
+				if (LinkedExecNode.IsValid())
+				{
+					TunnelCallNode->AddExitSite(CurrExitPoint, ExitScriptCodeOffset);
+					if (CurrExitPoint->HasFlags(EScriptExecutionNodeFlags::ReEntrantTunnelPin))
 					{
-						const int32 ExitScriptCodeOffset = GetCodeLocationFromPin(*ExitPin);
-						if (CurrExitPoint->HasFlags(EScriptExecutionNodeFlags::ReEntrantTunnelPin))
-						{
-							CurrExitPoint->AddLinkedNode(ExitScriptCodeOffset, LinkedExecNode);
-						}
-						else
-						{
-							LinkedExecNode->AddFlags(EScriptExecutionNodeFlags::TunnelThenPin);
-							TunnelCallNode->AddLinkedNode(ExitScriptCodeOffset, LinkedExecNode);
-						}
+						CurrExitPoint->AddLinkedNode(ExitScriptCodeOffset, LinkedExecNode);
+					}
+					else
+					{
+						CurrExitPoint->AddFlags(EScriptExecutionNodeFlags::TunnelFinalExitPin);
+						LinkedExecNode->AddFlags(EScriptExecutionNodeFlags::TunnelFinalExitPin);
+						TunnelCallNode->AddLinkedNode(ExitScriptCodeOffset, LinkedExecNode);
 					}
 				}
 			}
@@ -832,6 +884,16 @@ TSharedPtr<FScriptExecutionNode> FBlueprintFunctionContext::GetProfilerDataForNo
 	return Result;
 }
 
+template<typename ExecNodeType> TSharedPtr<ExecNodeType> FBlueprintFunctionContext::GetTypedProfilerDataForNode(const FName NodeName)
+{
+	TSharedPtr<ExecNodeType> Result;
+	if (TSharedPtr<FScriptExecutionNode>* SearchResult = ExecutionNodes.Find(NodeName))
+	{
+		Result = StaticCastSharedPtr<ExecNodeType>(*SearchResult);
+	}
+	return Result;
+}
+
 const UEdGraphNode* FBlueprintFunctionContext::GetNodeFromCodeLocation(const int32 ScriptOffset)
 {
 	TWeakObjectPtr<const UEdGraphNode>& Result = ScriptOffsetToNodes.FindOrAdd(ScriptOffset);
@@ -900,6 +962,17 @@ TSharedPtr<FScriptExecutionNode> FBlueprintFunctionContext::CreateExecutionNode(
 	return ScriptExecNode;
 }
 
+template<typename ExecNodeType> TSharedPtr<ExecNodeType> FBlueprintFunctionContext::CreateTypedExecutionNode(FScriptExecNodeParams& InitParams)
+{
+	TSharedPtr<FScriptExecutionNode>& ScriptExecNode = ExecutionNodes.FindOrAdd(InitParams.NodeName);
+	check(!ScriptExecNode.IsValid());
+	UEdGraph* OuterGraph = InitParams.ObservedObject.IsValid() ? InitParams.ObservedObject.Get()->GetTypedOuter<UEdGraph>() : nullptr;
+	InitParams.OwningGraphName = OuterGraph ? OuterGraph->GetFName() : NAME_None;
+	TSharedPtr<ExecNodeType> NewTypedNode = MakeShareable(new ExecNodeType(InitParams));
+	ScriptExecNode = NewTypedNode;
+	return NewTypedNode;
+}
+
 void FBlueprintFunctionContext::AddCallSiteEntryPointsToNode(TSharedPtr<FScriptExecutionNode> CallingNode) const
 {
 	if (CallingNode.IsValid())
@@ -972,6 +1045,10 @@ void FBlueprintFunctionContext::MapTunnelInstance(UK2Node_Tunnel* TunnelInstance
 			}
 			TunnelGraphContext->StagingTunnelName = NAME_None;
 		}
+		// Setup tunnel graph context from instance.
+		TunnelGraphContext->BlueprintContext = BlueprintContext;
+		TunnelGraphContext->Function = Function;
+		TunnelGraphContext->BlueprintClass = BlueprintClass;
 		// Patch info into the calling function context
 		for (auto ExecutionNode : TunnelGraphContext->ExecutionNodes)
 		{
@@ -1041,7 +1118,7 @@ void FBlueprintFunctionContext::MapTunnelIO(UEdGraph* TunnelGraph)
 					LinkParams.ObservedObject = Pin;
 					LinkParams.DisplayName = Pin->GetDisplayName().IsEmpty() ? FText::FromString(*Pin->PinName) : Pin->GetDisplayName();
 					LinkParams.Tooltip = LOCTEXT("ExecPin_ToolTip", "Expand tunnel entry point");
-					LinkParams.NodeFlags = EScriptExecutionNodeFlags::TunnelPin;
+					LinkParams.NodeFlags = EScriptExecutionNodeFlags::TunnelEntryPin;
 					LinkParams.IconColor = FLinearColor(1.f, 1.f, 1.f, 0.8f);
 					const FSlateBrush* Icon = FEditorStyle::GetBrush(TEXT("BlueprintProfiler.BPPinConnected"));
 					LinkParams.Icon = const_cast<FSlateBrush*>(Icon);
@@ -1060,7 +1137,7 @@ void FBlueprintFunctionContext::MapTunnelIO(UEdGraph* TunnelGraph)
 					LinkParams.IconColor = FLinearColor(1.f, 1.f, 1.f, 0.8f);
 					const FSlateBrush* Icon = FEditorStyle::GetBrush(TEXT("BlueprintProfiler.BPPinConnected"));
 					LinkParams.Icon = const_cast<FSlateBrush*>(Icon);
-					TSharedPtr<FScriptExecutionNode> ExitPoint = CreateExecutionNode(LinkParams);
+					TSharedPtr<FScriptExecutionNode> ExitPoint = CreateTypedExecutionNode<FScriptExecutionTunnelExit>(LinkParams);
 					AddExitPoint(ExitPoint);
 				}
 			}
@@ -1092,6 +1169,26 @@ void FBlueprintFunctionContext::AddExitPoint(TSharedPtr<FScriptExecutionNode> Ex
 	TunnelIds.Add(ExitPoint->GetName()) = TunnelId;
 }
 
+TSharedPtr<FScriptExecutionTunnelEntry> FBlueprintFunctionContext::GetTunnelEntrySite(const int32 ScriptCodeOffset)
+{
+	TSharedPtr<FScriptExecutionTunnelEntry> Result;
+	if (TSharedPtr<FScriptExecutionTunnelEntry>* SearchResult = TunnelEntrySites.Find(ScriptCodeOffset))
+	{
+		Result = *SearchResult;
+	}
+	return Result;
+}
+
+TSharedPtr<FScriptExecutionTunnelEntry> FBlueprintFunctionContext::GetTunnelFromExitSite(const int32 ScriptCodeOffset)
+{
+	TSharedPtr<FScriptExecutionTunnelEntry> Result;
+	if (TSharedPtr<FScriptExecutionTunnelEntry>* SearchResult = TunnelExitSites.Find(ScriptCodeOffset))
+	{
+		Result = *SearchResult;
+	}
+	return Result;
+}
+
 //////////////////////////////////////////////////////////////////////////
 // FScriptEventPlayback
 
@@ -1105,12 +1202,6 @@ bool FScriptEventPlayback::Process(const TArray<FScriptInstrumentedEvent>& Signa
 		TSharedPtr<FBlueprintFunctionContext> FunctionContext;
 		TArray<FScriptInstrumentedEvent> UnproccessedEvents;
 	};
-	struct TunnelEventHelper
-	{
-		double TunnelEntryTime;
-		FTracePath EntryTracePath;
-		TSharedPtr<FScriptExecutionNode> TunnelEntryPoint;
-	};
 
 	const int32 NumEvents = (StopIdx+1) - StartIdx;
 	ProcessingState = EEventProcessingResult::Failed;
@@ -1125,7 +1216,6 @@ bool FScriptEventPlayback::Process(const TArray<FScriptInstrumentedEvent>& Signa
 		TSharedPtr<FScriptExecutionNode> EventNode = FunctionContext->GetProfilerDataForNode(EventName);
 		// Find Associated graph nodes and submit into a node map for later processing.
 		TMap<const UEdGraphNode*, PerNodeEventHelper> NodeMap;
-		TArray<TunnelEventHelper> Tunnels;
 		ProcessingState = EEventProcessingResult::Success;
 		int32 LastEventIdx = StartIdx;
 		const int32 EventStartOffset = SignalData[StartIdx].IsResumeEvent() ? 3 : 1;
@@ -1165,6 +1255,7 @@ bool FScriptEventPlayback::Process(const TArray<FScriptInstrumentedEvent>& Signa
 						break;
 					}
 					case EScriptInstrumentation::NodeEntry:
+					case EScriptInstrumentation::NodeDebugSite:
 					{
 						// Handle timings for events called as functions
 						if (NodeEntry.ExecNode->IsCustomEvent())
@@ -1183,6 +1274,9 @@ bool FScriptEventPlayback::Process(const TArray<FScriptInstrumentedEvent>& Signa
 								}
 							}
 						}
+						// Process tunnel timings
+						ProcessTunnelEntryPoints(FunctionContext, CurrSignal);
+						// Queue events
 						NodeEntry.EntryTracePaths.Push(TracePath);
 						NodeEntry.UnproccessedEvents.Add(CurrSignal);
 						break;
@@ -1221,7 +1315,7 @@ bool FScriptEventPlayback::Process(const TArray<FScriptInstrumentedEvent>& Signa
 							{
 								const FScriptInstrumentedEvent& NodeEvent = NodeEntry.UnproccessedEvents[NodeEventIdx];
 
-								if (NodeEvent.GetType() == EScriptInstrumentation::NodeEntry)
+								if (NodeEvent.GetType() == EScriptInstrumentation::NodeEntry || NodeEvent.GetType() == EScriptInstrumentation::NodeDebugSite)
 								{
 									if (NodeEntryIdx != INDEX_NONE)
 									{
@@ -1241,6 +1335,15 @@ bool FScriptEventPlayback::Process(const TArray<FScriptInstrumentedEvent>& Signa
 					}
 					case EScriptInstrumentation::NodeExit:
 					{
+						// Handle duplicate exit sites for If Then Else nodes, This should be cleaned up in the refactor.
+						if (GraphNode->IsA<UK2Node_IfThenElse>())
+						{
+							const int32 NextSignal = SignalIdx + 1;
+							if (NextSignal < StopIdx && SignalData[NextSignal].GetType() == EScriptInstrumentation::NodeExit)
+							{
+								continue;
+							}
+						}
 						// Cleanup branching multiple exits and correct the tracepath
 						if (NodeEntry.UnproccessedEvents.Num() && NodeEntry.UnproccessedEvents.Last().GetType() == EScriptInstrumentation::NodeExit)
 						{
@@ -1250,33 +1353,10 @@ bool FScriptEventPlayback::Process(const TArray<FScriptInstrumentedEvent>& Signa
 								TracePath = NodeEntry.EntryTracePaths.Last();
 							}
 						}
-						const int32 ScriptCodeExit = CurrSignal.GetScriptCodeOffset();
 						// Process tunnel timings
-						for (auto LinkIter : NodeEntry.ExecNode->GetLinkedNodes())
-						{
-							if (LinkIter.Key == ScriptCodeExit && LinkIter.Value->HasFlags(EScriptExecutionNodeFlags::FunctionTunnel))
-							{
-								TunnelEventHelper NewTunnel;
-								NewTunnel.EntryTracePath = TracePath;
-								NewTunnel.EntryTracePath.AddExitPin(ScriptCodeExit);
-								NewTunnel.TunnelEntryPoint = LinkIter.Value;
-								NewTunnel.TunnelEntryTime = CurrSignal.GetTime();
-								Tunnels.Push(NewTunnel);
-							}
-							else if (LinkIter.Value->HasFlags(EScriptExecutionNodeFlags::TunnelExitPin))
-							{
-								if (Tunnels.Num())
-								{
-									TunnelEventHelper Tunnel = Tunnels.Pop();
-									TSharedPtr<FScriptPerfData> PerfData = Tunnel.TunnelEntryPoint->GetPerfDataByInstanceAndTracePath(InstanceName, Tunnel.EntryTracePath);
-									PerfData->AddEventTiming(CurrSignal.GetTime() - Tunnel.TunnelEntryTime);
-									FTracePath ExitTrace(TracePath);
-									ExitTrace.AddExitPin(LinkIter.Key);
-									PerfData = LinkIter.Value->GetPerfDataByInstanceAndTracePath(InstanceName, ExitTrace);
-									PerfData->AddEventTiming(0.0);
-								}
-							}
-						}
+						ProcessTunnelEntryPoints(FunctionContext, CurrSignal);
+						// Process node exit
+						const int32 ScriptCodeExit = CurrSignal.GetScriptCodeOffset();
 						TracePath.AddExitPin(ScriptCodeExit);
 						NodeEntry.UnproccessedEvents.Add(CurrSignal);
 						// Process cyclic linkage - reset to root link tracepath
@@ -1370,6 +1450,7 @@ bool FScriptEventPlayback::Process(const TArray<FScriptInstrumentedEvent>& Signa
 						break;
 					}
 					case EScriptInstrumentation::NodeEntry:
+					case EScriptInstrumentation::NodeDebugSite:
 					{
 						if (PureNode.IsValid())
 						{
@@ -1411,4 +1492,41 @@ bool FScriptEventPlayback::Process(const TArray<FScriptInstrumentedEvent>& Signa
 	return ProcessingState != EEventProcessingResult::Failed;
 }
 
+void FScriptEventPlayback::ProcessTunnelEntryPoints(TSharedPtr<FBlueprintFunctionContext> FunctionContext, const FScriptInstrumentedEvent& CurrSignal)
+{
+	const int32 ScriptOffset = CurrSignal.GetScriptCodeOffset();
+	if (const UEdGraphPin* TunnelPin = FunctionContext->GetPinFromCodeLocation(ScriptOffset))
+	{
+		if (const UK2Node_Tunnel* TunnelInstance = Cast<UK2Node_Tunnel>(TunnelPin->GetOwningNode()))
+		{
+			TSharedPtr<FScriptExecutionTunnelEntry> ExitSite = FunctionContext->GetTunnelFromExitSite(ScriptOffset);
+			if (ExitSite.IsValid())
+			{
+				if (TunnelEventHelper* OpenTunnel = ActiveTunnels.Find(ExitSite->GetName()))
+				{
+					const double TunnelTiming = CurrSignal.GetTime() - OpenTunnel->TunnelEntryTime;
+					TSharedPtr<FScriptExecutionTunnelEntry> TunnelEntryNode = OpenTunnel->TunnelEntryPoint;
+					TunnelEntryNode->UpdateTunnelExit(InstanceName, TracePath, ScriptOffset);
+					if (TunnelEntryNode->IsFinalExitSite(ScriptOffset))
+					{
+						TunnelEntryNode->AddTunnelTiming(InstanceName, OpenTunnel->EntryTracePath, TunnelTiming);
+						TracePath = OpenTunnel->EntryTracePath;
+						ActiveTunnels.Remove(ExitSite->GetName());
+					}
+				}
+			}
+			// Handle tunnel entries
+			TSharedPtr<FScriptExecutionTunnelEntry> EntrySite = FunctionContext->GetTunnelEntrySite(ScriptOffset);
+			if (EntrySite.IsValid())
+			{
+				TunnelEventHelper& NewTunnel = ActiveTunnels.Add(EntrySite->GetName());
+				NewTunnel.EntryTracePath = TracePath;
+				NewTunnel.EntryTracePath.AddExitPin(ScriptOffset);
+				NewTunnel.TunnelEntryPoint = EntrySite;
+				NewTunnel.TunnelEntryTime = CurrSignal.GetTime();
+			}
+		}
+	}
+}
+
 #undef LOCTEXT_NAMESPACE
diff --git a/Engine/Source/Developer/BlueprintProfiler/Public/BlueprintProfiler.h b/Engine/Source/Developer/BlueprintProfiler/Public/BlueprintProfiler.h
--- a/Engine/Source/Developer/BlueprintProfiler/Public/BlueprintProfiler.h
+++ b/Engine/Source/Developer/BlueprintProfiler/Public/BlueprintProfiler.h
@@ -37,6 +37,8 @@ public:
 	virtual FName MapBlueprintInstance(TSharedPtr<FBlueprintExecutionContext> BlueprintContext, const FString& InstancePath) override;
 	virtual bool HasDataForInstance(const UObject* Instance) const override;
 	virtual void ProcessEventProfilingData() override;
+	virtual EBlueprintProfilerHeatMapDisplayMode::Type GetGraphNodeHeatMapDisplayMode() const { return GraphNodeHeatMapDisplayMode; }
+	virtual void SetGraphNodeHeatMapDisplayMode(EBlueprintProfilerHeatMapDisplayMode::Type InHeatMapDisplayMode) { GraphNodeHeatMapDisplayMode = InHeatMapDisplayMode; }
 #endif
 	// End IBlueprintProfilerModule
 
@@ -81,6 +83,8 @@ protected:
 #if WITH_EDITOR
 	/** PIE Active */
 	bool bPIEActive;
+	/** Current graph node heat map display mode */
+	EBlueprintProfilerHeatMapDisplayMode::Type GraphNodeHeatMapDisplayMode;
 	/** Suspended script events */
 	TMap<FName, TSharedPtr<FScriptEventPlayback>> SuspendedEvents;
 	/** Cached object path and code offset lookup to UObjects */
diff --git a/Engine/Source/Developer/BlueprintProfiler/Public/BlueprintProfilerModule.h b/Engine/Source/Developer/BlueprintProfiler/Public/BlueprintProfilerModule.h
--- a/Engine/Source/Developer/BlueprintProfiler/Public/BlueprintProfilerModule.h
+++ b/Engine/Source/Developer/BlueprintProfiler/Public/BlueprintProfilerModule.h
@@ -4,6 +4,22 @@
 
 #include "ModuleManager.h"
 
+#if WITH_EDITOR
+
+/** Profiler heat map display mode */
+namespace EBlueprintProfilerHeatMapDisplayMode
+{
+	enum Type
+	{
+		None = 0,
+		Inclusive,
+		Exclusive,
+		MaxTiming
+	};
+}
+
+#endif
+
 /** Delegate to broadcast structural stats changes */
 DECLARE_MULTICAST_DELEGATE_OneParam(FOnBPStatGraphLayoutChanged, TWeakObjectPtr<UBlueprint>);
 
@@ -46,6 +62,12 @@ public:
 	/** Process profiling event data */
 	virtual void ProcessEventProfilingData() {}
 
+	/** Returns the current graph node heat map display mode setting */
+	virtual EBlueprintProfilerHeatMapDisplayMode::Type GetGraphNodeHeatMapDisplayMode() const = 0;
+
+	/** Sets the current graph node heat map display mode setting */
+	virtual void SetGraphNodeHeatMapDisplayMode(EBlueprintProfilerHeatMapDisplayMode::Type InHeatMapDisplayMode) {}
+
 #endif // WITH_EDITOR
 
 };
\ No newline at end of file
diff --git a/Engine/Source/Developer/BlueprintProfiler/Public/ScriptInstrumentationPlayback.h b/Engine/Source/Developer/BlueprintProfiler/Public/ScriptInstrumentationPlayback.h
--- a/Engine/Source/Developer/BlueprintProfiler/Public/ScriptInstrumentationPlayback.h
+++ b/Engine/Source/Developer/BlueprintProfiler/Public/ScriptInstrumentationPlayback.h
@@ -123,12 +123,21 @@ public:
 	/** Returns the execution node representing the node name */
 	TSharedPtr<FScriptExecutionNode> GetProfilerDataForNode(const FName NodeName);
 
+	/** Returns the execution node representing the node name */
+	template<typename ExecNodeType> TSharedPtr<ExecNodeType> GetTypedProfilerDataForNode(const FName NodeName);
+
 	/** Adds all function entry points to node as children */
 	void AddCallSiteEntryPointsToNode(TSharedPtr<FScriptExecutionNode> CallingNode) const;
 
 	/** Returns the function name for this context */
 	FName GetFunctionName() const { return FunctionName; }
 
+	/** Returns the tunnel entry from the script code offset */
+	TSharedPtr<FScriptExecutionTunnelEntry> GetTunnelEntrySite(const int32 ScriptCodeOffset);
+
+	/** Returns the tunnel custom entry from the script code offset */
+	TSharedPtr<FScriptExecutionTunnelEntry> GetTunnelFromExitSite(const int32 ScriptCodeOffset);
+
 private:
 
 	friend FBlueprintExecutionContext;
@@ -136,6 +145,9 @@ private:
 	/** Utility to create execution node */
 	TSharedPtr<FScriptExecutionNode> CreateExecutionNode(FScriptExecNodeParams& InitParams);
 
+	/** Utility to create typed execution node */
+	template<typename ExecNodeType> TSharedPtr<ExecNodeType> CreateTypedExecutionNode(FScriptExecNodeParams& InitParams);
+
 	/** Adds and links in a new function entry point */
 	void AddEntryPoint(TSharedPtr<FScriptExecutionNode> EntryPoint);
 
@@ -206,6 +218,10 @@ private:
 	TMap<FName, TSharedPtr<FScriptExecutionNode>> TunnelEntryPointMap;
 	/** Map associating function entry points to exit points */
 	TMap<FName, TArray<TSharedPtr<FScriptExecutionNode>>> TunnelExitPointMap;
+	/** Map associating tunnel entry points to script code offsets */
+	TMap<int32, TSharedPtr<FScriptExecutionTunnelEntry>> TunnelEntrySites;
+	/** Map associating custom tunnel entries to script code offsets */
+	TMap<int32, TSharedPtr<FScriptExecutionTunnelEntry>> TunnelExitSites;
 	/** Staging tunnel name during mapping */
 	FName StagingTunnelName;
 	/** Map of tunnel Ids */
@@ -230,6 +246,13 @@ class FScriptEventPlayback : public TSharedFromThis<FScriptEventPlayback>
 {
 public:
 
+	struct TunnelEventHelper
+	{
+		double TunnelEntryTime;
+		FTracePath EntryTracePath;
+		TSharedPtr<FScriptExecutionTunnelEntry> TunnelEntryPoint;
+	};
+
 	FScriptEventPlayback(TSharedPtr<FBlueprintExecutionContext> BlueprintContextIn, const FName InstanceNameIn)
 		: InstanceName(InstanceNameIn)
 		, ProcessingState(EEventProcessingResult::None)
@@ -249,6 +272,11 @@ public:
 	/** Sets the blueprint context for this batch */
 	void SetBlueprintContext(TSharedPtr<FBlueprintExecutionContext> BlueprintContextIn) { BlueprintContext = BlueprintContextIn; }
 
+private:
+
+	/** Process tunnel entry points */
+	void ProcessTunnelEntryPoints(TSharedPtr<FBlueprintFunctionContext> FunctionContext, const FScriptInstrumentedEvent& CurrSignal);
+
 private:
 
 	/** Instance name */
@@ -261,6 +289,8 @@ private:
 	FTracePath TracePath;
 	/** Current tracepath stack */
 	TArray<FTracePath> TraceStack;
+	/** Active tunnels */
+	TMap<FName, TunnelEventHelper> ActiveTunnels;
 	/** Event Timings */
 	TMultiMap<FName, double> EventTimings;
 	/** Blueprint exec context */
diff --git a/Engine/Source/Editor/BlueprintGraph/Private/CallFunctionHandler.cpp b/Engine/Source/Editor/BlueprintGraph/Private/CallFunctionHandler.cpp
--- a/Engine/Source/Editor/BlueprintGraph/Private/CallFunctionHandler.cpp
+++ b/Engine/Source/Editor/BlueprintGraph/Private/CallFunctionHandler.cpp
@@ -129,13 +129,26 @@ void FKCHandler_CallFunction::CreateFunctionCallStatement(FKismetFunctionContext
 						{
 							check((*Term)->bIsLiteral);
 						
-							const int32 UUID = Context.GetContextUniqueID();
+							int32 LatentUUID = INDEX_NONE;
+							UEdGraphNode* AssociatedNode = LatentInfoPin->GetOwningNode();
+							if (AssociatedNode && AssociatedNode->NodeGuid.IsValid())
+							{
+								LatentUUID = GetTypeHash(AssociatedNode->NodeGuid);
+							}
+							else
+							{
+								static int32 FallbackUUID = 0;
+								LatentUUID = FallbackUUID++;
+
+								CompilerContext.MessageLog.Warning(*LOCTEXT("UUIDDeterministicCookWarn", "Failed to produce a deterministic UUID for a node's latent action: @@. Cooking this Blueprint (@@) asset will non-deterministic.").ToString(), LatentInfoPin, Context.Blueprint);
+							}							
+
 							const FString ExecutionFunctionName = CompilerContext.GetSchema()->FN_ExecuteUbergraphBase.ToString() + TEXT("_") + Context.Blueprint->GetName();
-							(*Term)->Name = FString::Printf(TEXT("(Linkage=%s,UUID=%s,ExecutionFunction=%s,CallbackTarget=None)"), *FString::FromInt(INDEX_NONE), *FString::FromInt(UUID), *ExecutionFunctionName);
+							(*Term)->Name = FString::Printf(TEXT("(Linkage=%s,UUID=%s,ExecutionFunction=%s,CallbackTarget=None)"), *FString::FromInt(INDEX_NONE), *FString::FromInt(LatentUUID), *ExecutionFunctionName);
 
 							// Record the UUID in the debugging information
 							UEdGraphNode* TrueSourceNode = Cast<UEdGraphNode>(Context.MessageLog.FindSourceObject(Node));
-							Context.NewClass->GetDebugData().RegisterUUIDAssociation(TrueSourceNode, UUID);
+							Context.NewClass->GetDebugData().RegisterUUIDAssociation(TrueSourceNode, LatentUUID);
 						}
 					}
 					else
diff --git a/Engine/Source/Editor/BlueprintGraph/Private/K2Node_MathExpression.cpp b/Engine/Source/Editor/BlueprintGraph/Private/K2Node_MathExpression.cpp
--- a/Engine/Source/Editor/BlueprintGraph/Private/K2Node_MathExpression.cpp
+++ b/Engine/Source/Editor/BlueprintGraph/Private/K2Node_MathExpression.cpp
@@ -2660,7 +2660,13 @@ void UK2Node_MathExpression::ClearExpression()
 		}
 	}
 	
-	CachedMessageLog = MakeShareable(new FCompilerResultsLog);
+	// passing true to FCompilerResultsLog's constructor would make this the
+	// primary compiler log (it is not) - the idea being that upon destruction 
+	// the primary log prints a summary; well, since this isn't destructed at 
+	// the end of compilation, and it blocks the full compiler log from 
+	// becoming the "CurrentEventTarget", we pass false - we append logs 
+	// collected by this one to the full compiler log later on anyways (so they won't be missed)
+	CachedMessageLog = MakeShareable(new FCompilerResultsLog(/*bIsCompatibleWithEvents =*/false));
 	
 	Expression.Empty();
 }
diff --git a/Engine/Source/Editor/EditorStyle/Private/SlateEditorStyle.cpp b/Engine/Source/Editor/EditorStyle/Private/SlateEditorStyle.cpp
--- a/Engine/Source/Editor/EditorStyle/Private/SlateEditorStyle.cpp
+++ b/Engine/Source/Editor/EditorStyle/Private/SlateEditorStyle.cpp
@@ -3688,6 +3688,7 @@ void FSlateEditorStyle::FStyle::SetupGraphEditorStyles()
 		Set( "Graph.Node.TitleGloss", new BOX_BRUSH( "/Graph/RegularNode_title_gloss", FMargin(12.0f/64.0f) ) );
 		Set( "Graph.Node.ColorSpill", new BOX_BRUSH( "/Graph/RegularNode_color_spill", FMargin(8.0f/64.0f, 3.0f/32.0f, 0, 0) ) );
 		Set( "Graph.Node.TitleHighlight", new BOX_BRUSH( "/Graph/RegularNode_title_highlight", FMargin(16.0f/64.0f, 1.0f, 16.0f/64.0f, 0.0f) ) );
+		Set( "Graph.Node.IndicatorOverlay", new IMAGE_BRUSH( "/Graph/IndicatorOverlay_color_spill", FVector2D(128.f, 32.f) ) );
 
 		Set( "Graph.Node.ShadowSize", FVector2D(12,12) );
 		Set( "Graph.Node.ShadowSelected", new BOX_BRUSH( "/Graph/RegularNode_shadow_selected", FMargin(18.0f/64.0f) ) );
@@ -3695,7 +3696,7 @@ void FSlateEditorStyle::FStyle::SetupGraphEditorStyles()
 
 		Set( "Graph.CompactNode.Body", new BOX_BRUSH( "/Graph/MathNode_body", FMargin(4.0f/64.0f, 26.0f/64.0f, 4.0f/64.0f, 4.0f/64.0f) ) );
 		Set( "Graph.CompactNode.ShadowSelected", new BOX_BRUSH( "/Graph/MathNode_shadow_selected", FMargin(18.0f/64.0f) ) );
-		
+		Set( "Graph.CompactNode.IndicatorOverlay", new IMAGE_BRUSH("/Graph/IndicatorOverlay_color_spill", FVector2D(64.f, 32.f)));
 
 		Set( "Graph.Node.CommentBubble", new BOX_BRUSH( "Old/Graph/CommentBubble", FMargin(8/32.0f) ) );
 		Set( "Graph.Node.CommentArrow", new IMAGE_BRUSH( "Old/Graph/CommentBubbleArrow", FVector2D(8,8) ) );
@@ -3739,6 +3740,7 @@ void FSlateEditorStyle::FStyle::SetupGraphEditorStyles()
 		Set( "Graph.VarNode.Body", new BOX_BRUSH( "/Graph/VarNode_body", FMargin(16.f/64.f, 12.f/28.f) ) );
 		Set( "Graph.VarNode.ColorSpill", new IMAGE_BRUSH( "/Graph/VarNode_color_spill", FVector2D(132,28) ) );
 		Set( "Graph.VarNode.Gloss", new BOX_BRUSH( "/Graph/VarNode_gloss", FMargin(16.f/64.f, 16.f/28.f, 16.f/64.f, 4.f/28.f) ) );
+		Set( "Graph.VarNode.IndicatorOverlay", new IMAGE_BRUSH("/Graph/IndicatorOverlay_color_spill", FVector2D(64.f, 28.f)));
 		
 		Set( "Graph.VarNode.ShadowSelected", new BOX_BRUSH( "/Graph/VarNode_shadow_selected", FMargin(26.0f/64.0f) ) );
 		Set( "Graph.VarNode.Shadow", new BOX_BRUSH( "/Graph/VarNode_shadow", FMargin(26.0f/64.0f) ) );
diff --git a/Engine/Source/Editor/GraphEditor/GraphEditor.Build.cs b/Engine/Source/Editor/GraphEditor/GraphEditor.Build.cs
--- a/Engine/Source/Editor/GraphEditor/GraphEditor.Build.cs
+++ b/Engine/Source/Editor/GraphEditor/GraphEditor.Build.cs
@@ -39,6 +39,7 @@ public class GraphEditor : ModuleRules
 				"UnrealEd",
 				"AssetRegistry",
 				"ClassViewer",
+                "Kismet",
 				"KismetWidgets",
 				"BlueprintGraph",
 				"AnimGraph",
diff --git a/Engine/Source/Editor/GraphEditor/Private/KismetNodes/SGraphNodeK2Base.cpp b/Engine/Source/Editor/GraphEditor/Private/KismetNodes/SGraphNodeK2Base.cpp
--- a/Engine/Source/Editor/GraphEditor/Private/KismetNodes/SGraphNodeK2Base.cpp
+++ b/Engine/Source/Editor/GraphEditor/Private/KismetNodes/SGraphNodeK2Base.cpp
@@ -10,6 +10,11 @@
 #include "IDocumentation.h"
 #include "TutorialMetaData.h"
 
+// Blueprint Profiler
+#include "Developer/BlueprintProfiler/Public/BlueprintProfilerModule.h"
+#include "Editor/Kismet/Public/Profiler/EventExecution.h"
+#include "Editor/Kismet/Public/Profiler/ScriptPerfData.h"
+
 #define LOCTEXT_NAMESPACE "SGraphNodeK2Base"
 
 //////////////////////////////////////////////////////////////////////////
@@ -80,8 +85,20 @@ void SGraphNodeK2Base::UpdateCompactNode()
 			+SOverlay::Slot()
 			.Padding(Settings->GetNonPinNodeBodyPadding())
 			[
-				SNew(SImage)
-				.Image( FEditorStyle::GetBrush("Graph.CompactNode.Body") )
+				SNew(SOverlay)
+				+SOverlay::Slot()
+				[
+					SNew(SImage)
+					.Image(FEditorStyle::GetBrush("Graph.CompactNode.Body"))
+					.ColorAndOpacity(this, &SGraphNodeK2Base::GetNodeBodyColor)
+				]
+				+SOverlay::Slot()
+				[
+					SNew(SImage)
+					.Image(FEditorStyle::GetBrush("Graph.CompactNode.IndicatorOverlay"))
+					.Visibility(this, &SGraphNodeK2Base::GetNodeIndicatorOverlayVisibility)
+					.ColorAndOpacity(this, &SGraphNodeK2Base::GetNodeIndicatorOverlayColor)
+				]
 			]
 			+SOverlay::Slot()
 			.HAlign(HAlign_Center)
@@ -266,6 +283,53 @@ FText SGraphNodeK2Base::GetToolTipHeading() const
 	return FText::GetEmpty();
 }
 
+FSlateColor SGraphNodeK2Base::GetNodeIndicatorOverlayColor() const
+{
+	float IntensityValue = 0.0f;
+	const float IntensityScale = 0.8f;
+
+	IBlueprintProfilerInterface& ProfilerModule = FModuleManager::LoadModuleChecked<IBlueprintProfilerInterface>("BlueprintProfiler");
+	if (ProfilerModule.IsProfilerEnabled())
+	{
+		TSharedPtr<class FScriptExecutionNode> ExecNode = ProfilerModule.GetProfilerDataForNode(GetNodeObj());
+		if (ExecNode.IsValid())
+		{
+			FScriptPerfData NodePerfData;
+			ExecNode->GetBlueprintPerfDataForAllTracePaths(NodePerfData);
+
+			switch (ProfilerModule.GetGraphNodeHeatMapDisplayMode())
+			{
+			case EBlueprintProfilerHeatMapDisplayMode::Exclusive:
+				IntensityValue = NodePerfData.GetNodeHeatLevel();
+				break;
+
+			case EBlueprintProfilerHeatMapDisplayMode::Inclusive:
+				IntensityValue = NodePerfData.GetInclusiveHeatLevel();
+				break;
+
+			case EBlueprintProfilerHeatMapDisplayMode::MaxTiming:
+				IntensityValue = NodePerfData.GetMaxTimeHeatLevel();
+				break;
+			}
+		}
+	}
+
+	return FLinearColor(0.9f, 0.0f, 0.0f, FMath::Clamp<float>(IntensityValue, 0.0f, 1.0f) * IntensityScale);
+}
+
+EVisibility SGraphNodeK2Base::GetNodeIndicatorOverlayVisibility() const
+{
+	IBlueprintProfilerInterface& ProfilerModule = FModuleManager::LoadModuleChecked<IBlueprintProfilerInterface>("BlueprintProfiler");
+	if (ProfilerModule.IsProfilerEnabled() && ProfilerModule.GetGraphNodeHeatMapDisplayMode() != EBlueprintProfilerHeatMapDisplayMode::None)
+	{
+		return EVisibility::Visible;
+	}
+	else
+	{
+		return EVisibility::Hidden;
+	}
+}
+
 /**
  * Update this GraphNode to match the data that it is observing
  */
diff --git a/Engine/Source/Editor/GraphEditor/Private/KismetNodes/SGraphNodeK2Base.h b/Engine/Source/Editor/GraphEditor/Private/KismetNodes/SGraphNodeK2Base.h
--- a/Engine/Source/Editor/GraphEditor/Private/KismetNodes/SGraphNodeK2Base.h
+++ b/Engine/Source/Editor/GraphEditor/Private/KismetNodes/SGraphNodeK2Base.h
@@ -19,6 +19,8 @@ public:
 protected :
 	//~ Begin SGraphNode Interface
 	virtual TSharedPtr<SToolTip> GetComplexTooltip() override;
+	virtual FSlateColor GetNodeIndicatorOverlayColor() const;
+	virtual EVisibility GetNodeIndicatorOverlayVisibility() const;
 	//~ End SGraphNode Interface
 
 	/** Set up node in 'standard' mode */
diff --git a/Engine/Source/Editor/GraphEditor/Private/KismetNodes/SGraphNodeK2Var.cpp b/Engine/Source/Editor/GraphEditor/Private/KismetNodes/SGraphNodeK2Var.cpp
--- a/Engine/Source/Editor/GraphEditor/Private/KismetNodes/SGraphNodeK2Var.cpp
+++ b/Engine/Source/Editor/GraphEditor/Private/KismetNodes/SGraphNodeK2Var.cpp
@@ -204,6 +204,13 @@ void SGraphNodeK2Var::UpdateGraphNode()
 				.ColorAndOpacity( this, &SGraphNodeK2Var::GetVariableColor )
 			]
 			+ SOverlay::Slot()
+			[
+				SNew(SImage)
+				.Image(FEditorStyle::GetBrush("Graph.VarNode.IndicatorOverlay"))
+				.Visibility(this, &SGraphNodeK2Var::GetNodeIndicatorOverlayVisibility)
+				.ColorAndOpacity(this, &SGraphNodeK2Var::GetNodeIndicatorOverlayColor)
+			]
+			+ SOverlay::Slot()
 			[
 				SNew(SImage)
 				.Image( FEditorStyle::GetBrush("Graph.VarNode.Gloss") )
diff --git a/Engine/Source/Editor/GraphEditor/Private/SGraphNode.cpp b/Engine/Source/Editor/GraphEditor/Private/SGraphNode.cpp
--- a/Engine/Source/Editor/GraphEditor/Private/SGraphNode.cpp
+++ b/Engine/Source/Editor/GraphEditor/Private/SGraphNode.cpp
@@ -838,7 +838,19 @@ void SGraphNode::UpdateGraphNode()
 					.HAlign(HAlign_Fill)
 					.VAlign(VAlign_Top)
 					[
-						CreateNodeContentArea()
+						SNew(SOverlay)
+						+SOverlay::Slot()
+						.VAlign(VAlign_Fill)
+						[
+							SNew(SImage)
+							.Image(FEditorStyle::GetBrush("Graph.Node.IndicatorOverlay"))
+							.Visibility(this, &SGraphNode::GetNodeIndicatorOverlayVisibility)
+							.ColorAndOpacity(this, &SGraphNode::GetNodeIndicatorOverlayColor)
+						]
+						+SOverlay::Slot()
+						[
+							CreateNodeContentArea()
+						]
 					]
 
 					+SVerticalBox::Slot()
@@ -1024,6 +1036,18 @@ void SGraphNode::CreatePinWidgets()
 	{
 		UEdGraphPin* CurPin = GraphNode->Pins[PinIndex];
 
+		if ( !ensureMsgf(CurPin->GetOuter() == GraphNode
+			, TEXT("Graph node ('%s' - %s) has an invalid %s pin: '%s'; (with a bad %s outer: '%s'); skiping creation of a widget for this pin.")
+			, *GraphNode->GetNodeTitle(ENodeTitleType::ListView).ToString()
+			, *GraphNode->GetPathName()
+			, (CurPin->Direction == EEdGraphPinDirection::EGPD_Input) ? TEXT("input") : TEXT("output")
+			,  CurPin->PinFriendlyName.IsEmpty() ? *CurPin->PinName : *CurPin->PinFriendlyName.ToString()
+			,  CurPin->GetOuter() ? *CurPin->GetOuter()->GetClass()->GetName() : TEXT("UNKNOWN")
+			,  CurPin->GetOuter() ? *CurPin->GetOuter()->GetPathName() : TEXT("NULL")) )
+		{
+			continue;
+		}
+
 		CreateStandardPinWidget(CurPin);
 	}
 }
diff --git a/Engine/Source/Editor/GraphEditor/Public/SGraphNode.h b/Engine/Source/Editor/GraphEditor/Public/SGraphNode.h
--- a/Engine/Source/Editor/GraphEditor/Public/SGraphNode.h
+++ b/Engine/Source/Editor/GraphEditor/Public/SGraphNode.h
@@ -126,6 +126,12 @@ public:
 	/** @return the tint for the node's title text */
 	FLinearColor GetNodeTitleTextColor() const;
 
+	/** @return the tint for the node's indicator overlay */
+	virtual FSlateColor GetNodeIndicatorOverlayColor() const { return FLinearColor::Transparent; }
+
+	/** @return visibility state of the indicator overlay */
+	virtual EVisibility GetNodeIndicatorOverlayVisibility() const { return EVisibility::Hidden; }
+
 	/** @return the tooltip to display when over the node */
 	FText GetNodeTooltip() const;
 	
diff --git a/Engine/Source/Editor/Kismet/Kismet.Build.cs b/Engine/Source/Editor/Kismet/Kismet.Build.cs
--- a/Engine/Source/Editor/Kismet/Kismet.Build.cs
+++ b/Engine/Source/Editor/Kismet/Kismet.Build.cs
@@ -24,6 +24,7 @@ public class Kismet : ModuleRules
 		PrivateDependencyModuleNames.AddRange(
 			new string[] {
                 "AppFramework",
+                "BlueprintRuntime",
 				"Core",
 				"CoreUObject",
 				"Slate",
@@ -40,6 +41,7 @@ public class Kismet : ModuleRules
 				"AnimGraph",
 				"PropertyEditor",
 				"SourceControl",
+                "SharedSettingsWidgets",
                 "InputCore",
 				"EngineSettings",
                 "Projects",
diff --git a/Engine/Source/Editor/Kismet/Private/BlueprintEditorModule.cpp b/Engine/Source/Editor/Kismet/Private/BlueprintEditorModule.cpp
--- a/Engine/Source/Editor/Kismet/Private/BlueprintEditorModule.cpp
+++ b/Engine/Source/Editor/Kismet/Private/BlueprintEditorModule.cpp
@@ -7,6 +7,7 @@
 #include "BlueprintEditor.h"
 #include "BlueprintEditorUtils.h"
 #include "BlueprintUtilities.h"
+#include "BlueprintWarningsConfigurationPanel.h"
 #include "Toolkits/ToolkitManager.h"
 #include "Editor/UnrealEd/Public/Kismet2/DebuggerCommands.h"
 #include "Editor/WorkspaceMenuStructure/Public/WorkspaceMenuStructureModule.h"
@@ -16,6 +17,7 @@
 #include "UObjectToken.h"
 #include "InstancedStaticMeshSCSEditorCustomization.h"
 #include "InstancedReferenceSubobjectHelper.h"
+#include "ISettingsModule.h"
 #include "UserDefinedStructureEditor.h"
 #include "BlueprintGraphPanelPinFactory.h"
 #include "SDockTab.h"
@@ -218,10 +220,28 @@ void FBlueprintEditorModule::StartupModule()
 	FEdGraphUtilities::RegisterVisualPinFactory(BlueprintGraphPanelPinFactory);
 
 	PrepareAutoGeneratedDefaultEvents();
+
+	if (ISettingsModule* SettingsModule = FModuleManager::GetModulePtr<ISettingsModule>("Settings"))
+	{
+		ConfigurationPanel = SNew(SBlueprintWarningsConfigurationPanel);
+		SettingsModule->RegisterSettings(
+			"Project", 
+			"Engine", 
+			"Blueprints",
+			LOCTEXT("RuntimeSettingsName", "Blueprints"),
+			LOCTEXT("RuntimeSettingsDescription", "Configure Blueprints runtime behavior"),
+			ConfigurationPanel.ToSharedRef());
+	}
+
 }
 
 void FBlueprintEditorModule::ShutdownModule()
 {
+	if (ISettingsModule* SettingsModule = FModuleManager::GetModulePtr<ISettingsModule>("Settings"))
+	{
+		SettingsModule->UnregisterSettings("Project", "Engine", "Blueprints");
+		ConfigurationPanel = TSharedPtr<SWidget>();
+	}
 	// we're intentionally leaking UndoRedoHandler because the GEditor may be garbage when ShutdownModule is called:
 
 	// Cleanup all information for auto generated default event nodes by this module
diff --git a/Engine/Source/Editor/Kismet/Private/BlueprintWarningsConfigurationPanel.cpp b/Engine/Source/Editor/Kismet/Private/BlueprintWarningsConfigurationPanel.cpp
new file mode 100644
--- /dev/null
+++ b/Engine/Source/Editor/Kismet/Private/BlueprintWarningsConfigurationPanel.cpp
@@ -0,0 +1,258 @@
+// Copyright 1998-2016 Epic Games, Inc. All Rights Reserved.
+
+#include "BlueprintEditorPrivatePCH.h"
+
+#include "BlueprintRuntime.h"
+#include "BlueprintWarningsConfigurationPanel.h"
+#include "SSettingsEditorCheckoutNotice.h"
+
+#define LOCTEXT_NAMESPACE "BlueprintWarningConfigurationPanel"
+
+static const TCHAR* ColumnWarningIdentifier = TEXT("WarningIdentifier");
+static const TCHAR* ColumnWarningDescription = TEXT("WarningDescription");
+static const TCHAR* ColumnWarningBehavior = TEXT("WarningAsError");
+
+typedef SComboBox< TSharedPtr<EBlueprintWarningBehavior> > FBlueprintWarningBehaviorComboBox;
+
+class SBlueprintWarningRow : public SMultiColumnTableRow< FBlueprintWarningListEntry >
+{
+	SLATE_BEGIN_ARGS(SBlueprintWarningRow) { }
+		SLATE_STYLE_ARGUMENT( FTableRowStyle, Style )
+	SLATE_END_ARGS()
+	
+	void Construct( const FArguments& InArgs, const TSharedRef<STableViewBase>& InOwnerTableView, FBlueprintWarningListEntry InWarningInfo, SBlueprintWarningsConfigurationPanel* InParent )
+	{
+		WarningInfo = InWarningInfo;
+		Parent = InParent;
+
+		SMultiColumnTableRow<FBlueprintWarningListEntry>::Construct(FSuperRowType::FArguments(), InOwnerTableView);
+	}
+
+	virtual TSharedRef<SWidget> GenerateWidgetForColumn(const FName& InColumnName) override
+	{
+		if(InColumnName == ColumnWarningIdentifier)
+		{
+			return SNew(SBox)
+				.Padding(FMargin(4.0f, 0.0f))
+				.VAlign(VAlign_Center)
+				[
+					SNew(STextBlock)
+					.Text(FText::FromName(WarningInfo->WarningIdentifier))
+				];
+		}
+		else if (InColumnName == ColumnWarningDescription )
+		{
+			return SNew(SBox)
+				.Padding(FMargin(4.0f, 0.0f))
+				.VAlign(VAlign_Center)
+				[
+					SNew(STextBlock)
+					.Text(WarningInfo->WarningDescription)
+				];
+		}
+		else if (InColumnName == ColumnWarningBehavior )
+		{
+			return SNew(FBlueprintWarningBehaviorComboBox)
+				.Content()
+				[
+					SNew(STextBlock)
+						.Text_Lambda([this]() 
+						{
+							UEnum* const BlueprintWarningBehaviorEnum = FindObjectChecked<UEnum>(ANY_PACKAGE, TEXT("EBlueprintWarningBehavior"));
+							EBlueprintWarningBehavior Behavior = EBlueprintWarningBehavior::Warn;
+							FName WarningIdentifier = this->WarningInfo->WarningIdentifier;
+							if( FBlueprintSupport::ShouldTreatWarningAsError(WarningIdentifier))
+							{
+								Behavior = EBlueprintWarningBehavior::Error;
+							}
+							else if(FBlueprintSupport::ShouldSuppressWarning(WarningIdentifier))
+							{
+								Behavior = EBlueprintWarningBehavior::Suppress;
+							}
+							return BlueprintWarningBehaviorEnum->GetDisplayNameText(static_cast<int32>(Behavior));
+						})
+				]
+				.OptionsSource(&(Parent->CachedBlueprintWarningBehaviors))
+				.OnSelectionChanged(
+					FBlueprintWarningBehaviorComboBox::FOnSelectionChanged::CreateLambda(
+						[this]( TSharedPtr<EBlueprintWarningBehavior> Behavior, ESelectInfo::Type )
+						{
+							// apply result to each selected entry:
+							this->Parent->UpdateSelectedWarningBehaviors(*Behavior, *(this->WarningInfo));
+						}
+					)
+				)
+				.OnGenerateWidget(
+					FBlueprintWarningBehaviorComboBox::FOnGenerateWidget::CreateStatic(
+						[]( TSharedPtr<EBlueprintWarningBehavior> Behavior )->TSharedRef<SWidget>
+						{
+							UEnum* const BlueprintWarningBehaviorEnum = FindObjectChecked<UEnum>(ANY_PACKAGE, TEXT("EBlueprintWarningBehavior"));
+							return SNew(STextBlock)
+								.Text(BlueprintWarningBehaviorEnum->GetDisplayNameText(static_cast<int32>(*Behavior)));
+						}
+					)
+				);
+		}
+		else
+		{
+			ensure(false);
+			return SNew(SBorder);
+		}
+	}
+
+private:
+	FBlueprintWarningListEntry WarningInfo;
+	SBlueprintWarningsConfigurationPanel* Parent;
+};
+
+void SBlueprintWarningsConfigurationPanel::Construct(const FArguments& InArgs)
+{
+	UBlueprintRuntimeSettings* RuntimeSettings = GetMutableDefault<UBlueprintRuntimeSettings>();
+
+	const TArray<FBlueprintWarningDeclaration>& Warnings = FBlueprintSupport::GetBlueprintWarnings();
+	for (const FBlueprintWarningDeclaration& Warning : Warnings )
+	{
+		CachedBlueprintWarningData.Add(MakeShareable(new FBlueprintWarningDeclaration(Warning)));
+	}
+
+	// This is not a great pattern:
+	CachedBlueprintWarningBehaviors.Add( MakeShareable( new EBlueprintWarningBehavior( EBlueprintWarningBehavior::Warn ) ) );
+	CachedBlueprintWarningBehaviors.Add( MakeShareable( new EBlueprintWarningBehavior( EBlueprintWarningBehavior::Error ) ) );
+	CachedBlueprintWarningBehaviors.Add( MakeShareable( new EBlueprintWarningBehavior( EBlueprintWarningBehavior::Suppress ) ) );
+	
+	FString RelativeConfigFilePath = RuntimeSettings->GetDefaultConfigFilename();
+	FString FullSettingsPath = FPaths::ConvertRelativePathToFull(RelativeConfigFilePath);
+
+	TSharedPtr<SSettingsEditorCheckoutNotice> SettingsFile;
+	TSharedPtr<SBorder> Label;
+	// display a table of all known blueprint warnings, their description, etc:
+	ChildSlot
+	[
+		SNew(SVerticalBox)
+		+ SVerticalBox::Slot()
+			.AutoHeight()
+		[
+			SAssignNew(SettingsFile, SSettingsEditorCheckoutNotice)
+			.ConfigFilePath(FullSettingsPath)
+		]
+		+ SVerticalBox::Slot()
+			.AutoHeight()
+			.Padding(0.0f, 16.0f, 0.0f, 0.0f)
+		[
+			SAssignNew(Label, SBorder)
+				.Padding(3)
+				.BorderImage(FEditorStyle::GetBrush("DetailsView.CategoryTop"))
+				.BorderBackgroundColor(FLinearColor(.6, .6, .6, 1.0f))
+			[
+				SNew(STextBlock)
+				.Font(FEditorStyle::GetFontStyle("DetailsView.CategoryFontStyle"))
+				.Text(LOCTEXT("BlueprintWarningSettings", "Warning Behavior"))
+			]
+		]
+		+ SVerticalBox::Slot()
+		.AutoHeight()
+		[
+			SNew(SBorder)
+				.BorderImage(FEditorStyle::GetBrush("ToolPanel.GroupBorder"))
+			[
+				SAssignNew(ListView, FBlueprintWarningListView)
+					.SelectionMode(ESelectionMode::Multi)
+					.ListItemsSource(&CachedBlueprintWarningData)
+					.OnGenerateRow(
+						FBlueprintWarningListView::FOnGenerateRow::CreateLambda(
+							[this](FBlueprintWarningListEntry Warning, const TSharedRef<STableViewBase>& Owner) -> TSharedRef< ITableRow >
+							{
+								return SNew(SBlueprintWarningRow, Owner, Warning, this);
+							}
+						)
+					)
+					.HeaderRow
+					(
+						SNew(SHeaderRow)
+						+ SHeaderRow::Column(ColumnWarningIdentifier)
+							.DefaultLabel(LOCTEXT("BlueprintWarningIdentifierHeaderLabel", "Identifier"))
+							.DefaultTooltip(LOCTEXT("BlueprintWarningIdentifierHeaderTooltip", "Identifier used in game runtime when warning is raised"))
+							.FillWidth(0.15f)
+						+ SHeaderRow::Column(ColumnWarningDescription)
+							.DefaultLabel(LOCTEXT("BlueprintWarningDescriptionHeaderLabel", "Description"))
+							.DefaultTooltip(LOCTEXT("BlueprintWarningDescriptionHeaderTooltip", "Description of when the warning is raised"))
+							.FillWidth(0.55f)
+						+ SHeaderRow::Column(ColumnWarningBehavior)
+							.DefaultLabel(LOCTEXT("BlueprintWarningBehaviorHeaderLabel", "Behavior"))
+							.DefaultTooltip(LOCTEXT("BlueprintWarningBehaviorHeaderTooltip", "Determines what happens when the warning is raised - warnings can be treated more strictly or suppressed entirely"))
+							.FillWidth(0.3f)
+					)
+			]
+		]
+		+ SVerticalBox::Slot()
+		[
+			SNew(SSpacer)
+		]
+	];
+
+	TAttribute<bool> Enabled = TAttribute<bool>::Create(
+		TAttribute<bool>::FGetter::CreateLambda(
+			[SettingsFile]()
+			{
+				return SettingsFile->IsUnlocked();
+			}
+		)
+	);
+	ListView->SetEnabled(Enabled);
+	Label->SetEnabled(Enabled);
+}
+
+void SBlueprintWarningsConfigurationPanel::UpdateSelectedWarningBehaviors(EBlueprintWarningBehavior NewBehavior, const FBlueprintWarningDeclaration& AlteredWarning)
+{
+	UBlueprintRuntimeSettings* RuntimeSettings = GetMutableDefault<UBlueprintRuntimeSettings>();
+
+	// no TArray.ForEach makes this the most convenient pattern:
+	for( const TSharedPtr<FBlueprintWarningDeclaration>& Declaration : CachedBlueprintWarningData )
+	{
+		if( ListView->IsItemSelected(Declaration) || Declaration->WarningIdentifier == AlteredWarning.WarningIdentifier)
+		{
+			// update config data:
+			FName CurrentIdentifer = Declaration->WarningIdentifier;
+			FBlueprintWarningSettings* WarningEntry = RuntimeSettings->WarningSettings.FindByPredicate( 
+				[CurrentIdentifer]( const FBlueprintWarningSettings& Entry ) -> bool
+				{
+					return Entry.WarningIdentifier == CurrentIdentifer;
+				}
+			);
+			switch(NewBehavior)
+			{
+			case EBlueprintWarningBehavior::Warn:
+				if( WarningEntry )
+				{
+					RuntimeSettings->WarningSettings.RemoveAtSwap( WarningEntry - &(RuntimeSettings->WarningSettings[0]) );
+				}
+				break;
+			case EBlueprintWarningBehavior::Error:
+			case EBlueprintWarningBehavior::Suppress:
+				if( WarningEntry )
+				{
+					WarningEntry->WarningBehavior = NewBehavior;
+				}
+				else
+				{
+					FBlueprintWarningSettings NewEntry;
+					NewEntry.WarningIdentifier = CurrentIdentifer;
+					NewEntry.WarningDescription = Declaration->WarningDescription;
+					NewEntry.WarningBehavior = NewBehavior;
+					RuntimeSettings->WarningSettings.Add(
+						NewEntry
+					);
+				}
+				break;
+			}
+		}
+	}
+
+	// make sure runtime behavior matches config data:
+	if (IBlueprintRuntime* SettingsModule = FModuleManager::GetModulePtr<IBlueprintRuntime>("BlueprintRuntime"))
+	{
+		SettingsModule->PropagateWarningSettings();
+	}
+}
+
+#undef LOCTEXT_NAMESPACE
diff --git a/Engine/Source/Editor/Kismet/Private/BlueprintWarningsConfigurationPanel.h b/Engine/Source/Editor/Kismet/Private/BlueprintWarningsConfigurationPanel.h
new file mode 100644
--- /dev/null
+++ b/Engine/Source/Editor/Kismet/Private/BlueprintWarningsConfigurationPanel.h
@@ -0,0 +1,33 @@
+// Copyright 1998-2016 Epic Games, Inc. All Rights Reserved.
+
+#pragma once
+
+#include "BlueprintSupport.h"
+#include "BlueprintRuntimeSettings.h"
+
+typedef TSharedPtr<FBlueprintWarningDeclaration> FBlueprintWarningListEntry;
+typedef SListView<FBlueprintWarningListEntry> FBlueprintWarningListView;
+class SSettingsEditorCheckoutNotice;
+
+class SBlueprintWarningsConfigurationPanel : public SCompoundWidget
+{
+public:
+	SLATE_BEGIN_ARGS(SBlueprintWarningsConfigurationPanel){}
+	SLATE_END_ARGS()
+
+	void Construct(const FArguments& InArgs);
+private:
+	friend class SBlueprintWarningRow;
+
+	void UpdateSelectedWarningBehaviors(EBlueprintWarningBehavior NewBehavior, const FBlueprintWarningDeclaration& AlteredWarning );
+
+	// SListView requires TArray<TSharedRef<FFoo>> so we cache off a list from core:
+	TArray<TSharedPtr<FBlueprintWarningDeclaration>> CachedBlueprintWarningData;
+	// Again, SListView boilerplate:
+	TArray<TSharedPtr<EBlueprintWarningBehavior>> CachedBlueprintWarningBehaviors;
+	// Storing the listview, so we can apply updates to all selected entries:
+	TSharedPtr< FBlueprintWarningListView > ListView;
+	// Stored so that we can only enable controls when the settings files is writable:
+	TSharedPtr<SSettingsEditorCheckoutNotice> SettingsEditorCheckoutNotice;
+};
+
diff --git a/Engine/Source/Editor/Kismet/Private/Profiler/BPProfilerStatisticWidgets.cpp b/Engine/Source/Editor/Kismet/Private/Profiler/BPProfilerStatisticWidgets.cpp
--- a/Engine/Source/Editor/Kismet/Private/Profiler/BPProfilerStatisticWidgets.cpp
+++ b/Engine/Source/Editor/Kismet/Private/Profiler/BPProfilerStatisticWidgets.cpp
@@ -250,6 +250,16 @@ bool FBPProfilerStatDiplayOptions::IsFiltered(TSharedPtr<FScriptExecutionNode> N
 //////////////////////////////////////////////////////////////////////////
 // FBPProfilerStatWidget
 
+FBPProfilerStatWidget::FBPProfilerStatWidget(TSharedPtr<class FScriptExecutionNode> InExecNode, const FTracePath& WidgetTracePathIn)
+	: WidgetTracePath(WidgetTracePathIn)
+	, ExecNode(InExecNode)
+{
+	if (ExecNode->HasFlags(EScriptExecutionNodeFlags::FunctionTunnel))
+	{
+		WidgetTracePath = FTracePath(WidgetTracePathIn, StaticCastSharedPtr<FScriptExecutionTunnelEntry>(InExecNode));
+	}
+}
+
 TSharedRef<SWidget> FBPProfilerStatWidget::GenerateColumnWidget(FName ColumnName)
 {
 	if (ExecNode.IsValid())
@@ -307,7 +317,7 @@ TSharedRef<SWidget> FBPProfilerStatWidget::GenerateColumnWidget(FName ColumnName
 				if (ExecNode->HasFlags(NonNodeStats))
 				{
 					TextAttr = TAttribute<FText>(PerformanceStats.Get(), &FScriptPerfData::GetInclusiveTimingText);
-					ColorAttr = TAttribute<FSlateColor>(PerformanceStats.Get(), &FScriptPerfData::GetInclusiveHeatColor);
+					ColorAttr = TAttribute<FSlateColor>(this, &FBPProfilerStatWidget::GetInclusiveHeatColor);
 				}
 			}
 			else if (ColumnName == BlueprintProfilerStatText::ColumnId_Time)
@@ -315,7 +325,7 @@ TSharedRef<SWidget> FBPProfilerStatWidget::GenerateColumnWidget(FName ColumnName
 				if (!ExecNode->HasFlags(EScriptExecutionNodeFlags::ExecPin))
 				{
 					TextAttr = TAttribute<FText>(PerformanceStats.Get(), &FScriptPerfData::GetNodeTimingText);
-					ColorAttr = TAttribute<FSlateColor>(PerformanceStats.Get(), &FScriptPerfData::GetNodeHeatColor);
+					ColorAttr = TAttribute<FSlateColor>(this, &FBPProfilerStatWidget::GetNodeHeatColor);
 				}
 			}
 			else if (ColumnName == BlueprintProfilerStatText::ColumnId_MaxTime)
@@ -323,7 +333,7 @@ TSharedRef<SWidget> FBPProfilerStatWidget::GenerateColumnWidget(FName ColumnName
 				if (!ExecNode->HasFlags(EScriptExecutionNodeFlags::ExecPin))
 				{
 					TextAttr = TAttribute<FText>(PerformanceStats.Get(), &FScriptPerfData::GetMaxTimingText);
-					ColorAttr = TAttribute<FSlateColor>(PerformanceStats.Get(), &FScriptPerfData::GetMaxTimeHeatColor);
+					ColorAttr = TAttribute<FSlateColor>(this, &FBPProfilerStatWidget::GetMaxTimeHeatColor);
 				}
 			}
 			else if (ColumnName == BlueprintProfilerStatText::ColumnId_MinTime)
@@ -356,6 +366,24 @@ void FBPProfilerStatWidget::NavigateTo() const
 	}
 }
 
+FSlateColor FBPProfilerStatWidget::GetNodeHeatColor() const
+{
+	const float Value = 1.f - PerformanceStats->GetNodeHeatLevel();
+	return FLinearColor(1.f, Value, Value);
+}
+
+FSlateColor FBPProfilerStatWidget::GetInclusiveHeatColor() const
+{
+	const float Value = 1.f - PerformanceStats->GetInclusiveHeatLevel();
+	return FLinearColor(1.f, Value, Value);
+}
+
+FSlateColor FBPProfilerStatWidget::GetMaxTimeHeatColor() const
+{
+	const float Value = 1.f - PerformanceStats->GetMaxTimeHeatLevel();
+	return FLinearColor(1.f, Value, Value);
+}
+
 void FBPProfilerStatWidget::GenerateExecNodeWidgets(const TSharedPtr<FBPProfilerStatDiplayOptions> DisplayOptions)
 {
 	if (ExecNode.IsValid())
diff --git a/Engine/Source/Editor/Kismet/Private/Profiler/BPProfilerStatisticWidgets.h b/Engine/Source/Editor/Kismet/Private/Profiler/BPProfilerStatisticWidgets.h
--- a/Engine/Source/Editor/Kismet/Private/Profiler/BPProfilerStatisticWidgets.h
+++ b/Engine/Source/Editor/Kismet/Private/Profiler/BPProfilerStatisticWidgets.h
@@ -98,7 +98,7 @@ public:
 	bool HasAllFlags(const uint32 FlagsIn) const { return (Flags & FlagsIn) == FlagsIn; }
 
 	/** Set active instance name */
-	FName GetActiveInstance() const { return ActiveInstance; }
+	FName GetActiveInstance() const { return HasFlags(DisplayByInstance) ? ActiveInstance : NAME_None; }
 
 	/** Set active instance name */
 	void SetActiveInstance(const FName InstanceName);
@@ -135,11 +135,7 @@ class FBPProfilerStatWidget : public TSharedFromThis<FBPProfilerStatWidget>
 {
 public:
 
-	FBPProfilerStatWidget(TSharedPtr<class FScriptExecutionNode> InExecNode, const FTracePath& WidgetTracePathIn)
-		: WidgetTracePath(WidgetTracePathIn)
-		, ExecNode(InExecNode)
-	{
-	}
+	FBPProfilerStatWidget(TSharedPtr<class FScriptExecutionNode> InExecNode, const FTracePath& WidgetTracePathIn);
 
 	virtual ~FBPProfilerStatWidget() {}
 
@@ -178,6 +174,15 @@ protected:
 	/** Navigate to referenced node */
 	void NavigateTo() const;
 
+	/** Calculate and return the heat color to use for exclusive timing stats. */
+	FSlateColor GetNodeHeatColor() const;
+
+	/** Calculate and return the heat color to use for inclusive timing stats. */
+	FSlateColor GetInclusiveHeatColor() const;
+
+	/** Calculate and return the heat color to use for max time stats. */
+	FSlateColor GetMaxTimeHeatColor() const;
+
 protected:
 
 	/** Widget script tracepath */
diff --git a/Engine/Source/Editor/Kismet/Private/Profiler/EventExecution.cpp b/Engine/Source/Editor/Kismet/Private/Profiler/EventExecution.cpp
--- a/Engine/Source/Editor/Kismet/Private/Profiler/EventExecution.cpp
+++ b/Engine/Source/Editor/Kismet/Private/Profiler/EventExecution.cpp
@@ -26,6 +26,17 @@ TSharedPtr<FScriptExecutionNode> FScriptNodeExecLinkage::GetLinkedNodeByScriptOf
 //////////////////////////////////////////////////////////////////////////
 // FScriptNodePerfData
 
+void FScriptNodePerfData::GetValidInstanceNames(TSet<FName>& ValidInstances) const
+{
+	for (auto InstanceIter : InstanceInputPinToPerfDataMap)
+	{
+		if (InstanceIter.Key != NAME_None)
+		{
+			ValidInstances.Add(InstanceIter.Key);
+		}
+	}
+}
+
 bool FScriptNodePerfData::HasPerfDataForInstanceAndTracePath(FName InstanceName, const FTracePath& TracePath) const
 {
 	const TMap<const uint32, TSharedPtr<FScriptPerfData>>* InstanceMapPtr = InstanceInputPinToPerfDataMap.Find(InstanceName);
@@ -74,6 +85,21 @@ TSharedPtr<FScriptPerfData> FScriptNodePerfData::GetBlueprintPerfDataByTracePath
 	return Result;
 }
 
+void FScriptNodePerfData::GetBlueprintPerfDataForAllTracePaths(FScriptPerfData& OutPerfData)
+{
+	OutPerfData.Reset();
+
+	TMap<const uint32, TSharedPtr<FScriptPerfData>>& BlueprintMap = InstanceInputPinToPerfDataMap.FindOrAdd(NAME_None);
+	for (auto MapIt = BlueprintMap.CreateIterator(); MapIt; ++MapIt)
+	{
+		TSharedPtr<FScriptPerfData> CurDataPtr = MapIt.Value();
+		if (CurDataPtr.IsValid())
+		{
+			OutPerfData.AddData(*CurDataPtr);
+		}
+	}
+}
+
 //////////////////////////////////////////////////////////////////////////
 // FScriptExecutionNode
 
@@ -105,10 +131,6 @@ void FScriptExecutionNode::GetLinearExecutionPath(TArray<FLinearExecPath>& Linea
 {
 	LinearExecutionNodes.Add(FLinearExecPath(AsShared(), TracePath));
 	FTracePath NewTracePath(TracePath);
-	if (HasFlags(EScriptExecutionNodeFlags::FunctionTunnel))
-	{
-		MapTunnelLinearExecution(NewTracePath);
-	}
 	if (GetNumLinkedNodes() == 1)
 	{
 		for (auto NodeIter : LinkedNodes)
@@ -123,24 +145,7 @@ void FScriptExecutionNode::GetLinearExecutionPath(TArray<FLinearExecPath>& Linea
 				{
 					NewTracePath.AddExitPin(NodeIter.Key);
 				}
-				if (NodeIter.Value->HasFlags(EScriptExecutionNodeFlags::FunctionTunnel))
-				{
-					LinearExecutionNodes.Add(FLinearExecPath(NodeIter.Value, NewTracePath));
-					for (auto TunnelExit : NodeIter.Value->LinkedNodes)
-					{
-						if (TunnelExit.Value->HasFlags(EScriptExecutionNodeFlags::TunnelThenPin))
-						{
-							MapTunnelLinearExecution(NewTracePath);
-							NewTracePath.AddExitPin(TunnelExit.Key);
-							TunnelExit.Value->GetLinearExecutionPath(LinearExecutionNodes, NewTracePath);
-							break;
-						}
-					}
-				}
-				else
-				{
-					NodeIter.Value->GetLinearExecutionPath(LinearExecutionNodes, NewTracePath);
-				}
+				NodeIter.Value->GetLinearExecutionPath(LinearExecutionNodes, NewTracePath);
 			}
 		}
 	}
@@ -171,7 +176,7 @@ void FScriptExecutionNode::MapTunnelLinearExecution(FTracePath& Trace) const
 EScriptStatContainerType::Type FScriptExecutionNode::GetStatisticContainerType() const
 {
 	EScriptStatContainerType::Type Result = EScriptStatContainerType::Standard;
-	if (HasFlags(EScriptExecutionNodeFlags::ExecPin|EScriptExecutionNodeFlags::TunnelPin))
+	if (HasFlags(EScriptExecutionNodeFlags::ExecPin))
 	{
 		Result = EScriptStatContainerType::NewExecutionPath;
 	}
@@ -205,29 +210,15 @@ void FScriptExecutionNode::RefreshStats(const FTracePath& TracePath)
 		case EScriptStatContainerType::SequentialBranch:
 		{
 			GetInstancePerfDataByTracePath(TracePath, BlueprintPooledStats);
-
-			// Refresh attached pure chain stats (if present)
-			TSharedPtr<FScriptExecutionNode> PureChainNode = GetPureChainNode();
-			if (PureChainNode.IsValid())
-			{
-				PureChainNode->RefreshStats(TracePath);
-			}
 			break;
 		}
 		case EScriptStatContainerType::NewExecutionPath:
 		{
-			TSet<FName> AllInstances;
 			// Refresh child stats and copy as branch stats - this is a dummy entry
+			TSet<FName> AllInstances;
 			for (auto ChildIter : ChildNodes)
 			{
-				// Fill out instance data that can be missing because its a dummy node, or the code below can fail.
-				for (auto InstanceIter : ChildIter->InstanceInputPinToPerfDataMap)
-				{
-					if (InstanceIter.Key != NAME_None)
-					{
-						AllInstances.Add(InstanceIter.Key);
-					}
-				}
+				ChildIter->GetValidInstanceNames(AllInstances);
 				ChildIter->RefreshStats(TracePath);
 			}
 			for (auto InstanceIter : AllInstances)
@@ -269,31 +260,14 @@ void FScriptExecutionNode::RefreshStats(const FTracePath& TracePath)
 	// Refresh All links
 	if (bRefreshLinkStats)
 	{
-		if (IsPureChain())
+		for (auto LinkIter : LinkedNodes)
 		{
-			TMap<int32, TSharedPtr<FScriptExecutionNode>> AllPureNodes;
-			GetAllPureNodes(AllPureNodes);
-
-			FTracePath PureTracePath(TracePath);
-			for (auto PureIter : AllPureNodes)
+			FTracePath LinkTracePath(TracePath);
+			if (!LinkIter.Value->HasFlags(EScriptExecutionNodeFlags::InvalidTrace))
 			{
-				PureTracePath.AddExitPin(PureIter.Key);
-				PureIter.Value->RefreshStats(PureTracePath);
-			}
-		}
-		else
-		{
-			for (auto LinkIter : LinkedNodes)
-			{
-				FTracePath LinkTracePath(TracePath);
-				if (HasFlags(EScriptExecutionNodeFlags::FunctionTunnel))
-				{
-					MapTunnelLinearExecution(LinkTracePath);
-				}
-				//if (!LinkIter.Value->HasFlags(EScriptExecutionNodeFlags::DummyTrace))
 				LinkTracePath.AddExitPin(LinkIter.Key);
-				LinkIter.Value->RefreshStats(LinkTracePath);
 			}
+			LinkIter.Value->RefreshStats(LinkTracePath);
 		}
 	}
 	// Update the owning blueprint stats
@@ -367,6 +341,212 @@ void FScriptExecutionNode::NavigateToObject() const
 	}
 }
 
+//////////////////////////////////////////////////////////////////////////
+// FScriptExecutionTunnelInstance
+
+void FScriptExecutionTunnelInstance::GetBlueprintPerfDataForAllTracePaths(FScriptPerfData& OutPerfData)
+{
+	OutPerfData.Reset();
+
+	for (auto CustomEntryPoint : CustomEntryPoints)
+	{
+		CustomEntryPoint->GetBlueprintPerfDataForAllTracePaths(OutPerfData);
+	}
+}
+
+//////////////////////////////////////////////////////////////////////////
+// FScriptExecutionTunnelEntry
+
+void FScriptExecutionTunnelEntry::GetLinearExecutionPath(TArray<FLinearExecPath>& LinearExecutionNodes, const FTracePath& TracePath)
+{
+	LinearExecutionNodes.Add(FLinearExecPath(AsShared(), TracePath));
+	FTracePath TunnelTracePath(TracePath, SharedThis<const FScriptExecutionTunnelEntry>(this));
+	for (auto TunnelExit : LinkedNodes)
+	{
+		if (IsPathValidForTunnel(TunnelExit.Key))
+		{
+			TunnelTracePath.AddExitPin(TunnelExit.Key);
+			TunnelExit.Value->GetLinearExecutionPath(LinearExecutionNodes, TunnelTracePath);
+		}
+	}
+}
+
+void FScriptExecutionTunnelEntry::AddTunnelTiming(const FName InstanceName, const FTracePath& TracePath, const double Time)
+{
+	TSharedPtr<FScriptPerfData> PerfData = GetPerfDataByInstanceAndTracePath(InstanceName, TracePath);
+	PerfData->AddEventTiming(Time);
+	for (auto EntrySite : ChildNodes)
+	{
+		PerfData = EntrySite->GetPerfDataByInstanceAndTracePath(InstanceName, TracePath);
+		PerfData->TickSamples();
+	}
+}
+
+void FScriptExecutionTunnelEntry::UpdateTunnelExit(const FName InstanceName, const FTracePath& TracePath, const int32 ExitScriptOffset)
+{
+	// Tick exit site
+	TSharedPtr<FScriptExecutionNode> TunnelExitNode = GetExitSite(ExitScriptOffset);
+	if (TunnelExitNode.IsValid())
+	{
+		TSharedPtr<FScriptPerfData> PerfData = TunnelExitNode->GetPerfDataByInstanceAndTracePath(InstanceName, TracePath);
+		PerfData->TickSamples();
+	}
+}
+
+bool FScriptExecutionTunnelEntry::IsPathValidForTunnel(const int32 ScriptOffset) const
+{
+	return ValidExitPoints.Contains(ScriptOffset);
+}
+
+void FScriptExecutionTunnelEntry::AddExitSite(TSharedPtr<FScriptExecutionNode> ExitSite, const int32 ExitScriptOffset)
+{
+	ValidExitPoints.FindOrAdd(ExitScriptOffset) = ExitSite;
+}
+
+TSharedPtr<FScriptExecutionNode> FScriptExecutionTunnelEntry::GetExitSite(const int32 ScriptOffset) const
+{
+	TSharedPtr<FScriptExecutionNode> ExitSite;
+	if (const TSharedPtr<FScriptExecutionNode>* SearchResult = ValidExitPoints.Find(ScriptOffset))
+	{
+		ExitSite = *SearchResult;
+	}
+	return ExitSite;
+}
+
+bool FScriptExecutionTunnelEntry::IsFinalExitSite(const int32 ScriptOffset) const
+{
+	bool bResult = false;
+	if (const TSharedPtr<FScriptExecutionNode>* SearchResult = ValidExitPoints.Find(ScriptOffset))
+	{
+		bResult = (*SearchResult)->HasFlags(EScriptExecutionNodeFlags::TunnelFinalExitPin);
+	}
+	return bResult;
+}
+
+void FScriptExecutionTunnelEntry::RefreshStats(const FTracePath& TracePath)
+{
+	// Process stat update
+	TArray<TSharedPtr<FScriptPerfData>> BlueprintPooledStats;
+	// Update stats based on type
+	GetInstancePerfDataByTracePath(TracePath, BlueprintPooledStats);
+	// Refresh Child Links
+	for (auto ChildIter : ChildNodes)
+	{
+		ChildIter->RefreshStats(TracePath);
+	}
+	// Refresh All links
+	for (auto LinkIter : LinkedNodes)
+	{
+		FTracePath LinkTracePath(TracePath);
+		LinkTracePath.AddExitPin(LinkIter.Key);
+		LinkIter.Value->RefreshStats(LinkTracePath);
+	}
+	// Update the owning blueprint stats
+	if (BlueprintPooledStats.Num() > 0)
+	{
+		TSharedPtr<FScriptPerfData> BlueprintData = GetBlueprintPerfDataByTracePath(TracePath);
+		BlueprintData->Reset();
+		for (auto BlueprintChildDataIter : BlueprintPooledStats)
+		{
+			BlueprintData->AddData(*BlueprintChildDataIter.Get());
+		}
+	}
+}
+
+//////////////////////////////////////////////////////////////////////////
+// FScriptExecutionTunnelExit
+
+void FScriptExecutionTunnelExit::GetLinearExecutionPath(TArray<FLinearExecPath>& LinearExecutionNodes, const FTracePath& TracePath)
+{
+	LinearExecutionNodes.Add(FLinearExecPath(AsShared(), TracePath));
+	FTracePath NewTracePath(TracePath);
+	TSharedPtr<const FScriptExecutionTunnelEntry> ActiveTunnel = TracePath.GetTunnel();
+	if (ActiveTunnel.IsValid())
+	{
+		for (auto TunnelExit : LinkedNodes)
+		{
+			if (ActiveTunnel->IsPathValidForTunnel(TunnelExit.Key))
+			{
+				NewTracePath.AddExitPin(TunnelExit.Key);
+				TunnelExit.Value->GetLinearExecutionPath(LinearExecutionNodes, NewTracePath);
+			}
+		}
+	}
+}
+
+//////////////////////////////////////////////////////////////////////////
+// FScriptExecutionPureNode
+
+void FScriptExecutionPureNode::GetLinearExecutionPath(TArray<FLinearExecPath>& LinearExecutionNodes, const FTracePath& TracePath)
+{
+	LinearExecutionNodes.Add(FLinearExecPath(AsShared(), TracePath));
+	FTracePath NewTracePath(TracePath);
+	if (GetNumLinkedNodes() == 1)
+	{
+		for (auto NodeIter : LinkedNodes)
+		{
+			if (HasFlags(EScriptExecutionNodeFlags::PureStats))
+			{
+				continue;
+			}
+			else
+			{
+				if (NodeIter.Key != INDEX_NONE && !HasFlags(EScriptExecutionNodeFlags::InvalidTrace))
+				{
+					NewTracePath.AddExitPin(NodeIter.Key);
+				}
+				NodeIter.Value->GetLinearExecutionPath(LinearExecutionNodes, NewTracePath);
+			}
+		}
+	}
+}
+
+void FScriptExecutionPureNode::RefreshStats(const FTracePath& TracePath)
+{
+	// Process stat update
+	TArray<TSharedPtr<FScriptPerfData>> BlueprintPooledStats;
+	// Refresh all pure nodes and accumulate for blueprint stat.
+	if (HasFlags(EScriptExecutionNodeFlags::PureChain))
+	{
+		//GetInstancePerfDataByTracePath
+		TSet<FName> ValidInstances;
+		FTracePath PureTracePath(TracePath);
+		TMap<int32, TSharedPtr<FScriptExecutionNode>> AllPureNodes;
+		GetAllPureNodes(AllPureNodes);
+		for (auto PureIter : AllPureNodes)
+		{
+			GetValidInstanceNames(ValidInstances);
+			PureTracePath.AddExitPin(PureIter.Key);
+			PureIter.Value->RefreshStats(PureTracePath);
+		}
+		for (auto InstanceName : ValidInstances)
+		{
+			TSharedPtr<FScriptPerfData> InstancePerfData = GetPerfDataByInstanceAndTracePath(InstanceName, TracePath);
+			InstancePerfData->Reset();
+			BlueprintPooledStats.Add(InstancePerfData);
+
+			for (auto ChildIter : ChildNodes)
+			{
+				TSharedPtr<FScriptPerfData> ChildPerfData = ChildIter->GetPerfDataByInstanceAndTracePath(InstanceName, TracePath);
+				InstancePerfData->AddBranchData(*ChildPerfData.Get());
+			}
+		}
+	}
+	// Grab all instance stats to accumulate into the blueprint stat.
+	GetInstancePerfDataByTracePath(TracePath, BlueprintPooledStats);
+	// Update the owning blueprint stats
+	if (BlueprintPooledStats.Num() > 0)
+	{
+		TSharedPtr<FScriptPerfData> BlueprintData = GetBlueprintPerfDataByTracePath(TracePath);
+		BlueprintData->Reset();
+
+		for (auto BlueprintChildDataIter : BlueprintPooledStats)
+		{
+			BlueprintData->AddData(*BlueprintChildDataIter.Get());
+		}
+	}
+}
+
 //////////////////////////////////////////////////////////////////////////
 // FScriptExecutionBlueprint
 
diff --git a/Engine/Source/Editor/Kismet/Private/Profiler/SBlueprintProfilerView.cpp b/Engine/Source/Editor/Kismet/Private/Profiler/SBlueprintProfilerView.cpp
--- a/Engine/Source/Editor/Kismet/Private/Profiler/SBlueprintProfilerView.cpp
+++ b/Engine/Source/Editor/Kismet/Private/Profiler/SBlueprintProfilerView.cpp
@@ -84,6 +84,48 @@ void SBlueprintProfilerView::UpdateActiveProfilerWidget()
 										NAME_None, 
 										EUserInterfaceActionType::Button);
 
+		FMenuBuilder HeatMapDisplayModeComboContent(true, nullptr);
+		HeatMapDisplayModeComboContent.AddMenuEntry(LOCTEXT("HeatMapDisplayMode_None", "Off"),
+			LOCTEXT("HeatMapDisplayMode_None_Desc", "Normal Display (No Heat Map)"),
+			FSlateIcon(),
+			FUIAction(
+				FExecuteAction::CreateSP(this, &SBlueprintProfilerView::OnHeatMapDisplayModeChanged, EBlueprintProfilerHeatMapDisplayMode::None),
+				FCanExecuteAction(),
+				FIsActionChecked::CreateSP(this, &SBlueprintProfilerView::IsHeatMapDisplayModeSelected, EBlueprintProfilerHeatMapDisplayMode::None)
+			),
+			NAME_None,
+			EUserInterfaceActionType::Check);
+		HeatMapDisplayModeComboContent.AddMenuEntry(LOCTEXT("HeatMapDisplayMode_Exclusive", "Exclusive"),
+			LOCTEXT("HeatMapDisplayMode_Exclusive_Desc", "Heat Map - Exclusive Timings"),
+			FSlateIcon(),
+			FUIAction(
+				FExecuteAction::CreateSP(this, &SBlueprintProfilerView::OnHeatMapDisplayModeChanged, EBlueprintProfilerHeatMapDisplayMode::Exclusive),
+				FCanExecuteAction(),
+				FIsActionChecked::CreateSP(this, &SBlueprintProfilerView::IsHeatMapDisplayModeSelected, EBlueprintProfilerHeatMapDisplayMode::Exclusive)
+			),
+			NAME_None,
+			EUserInterfaceActionType::Check);
+		HeatMapDisplayModeComboContent.AddMenuEntry(LOCTEXT("HeatMapDisplayMode_Inclusive", "Inclusive"),
+			LOCTEXT("HeatMapDisplayMode_Inclusive_Desc", "Heat Map - Inclusive Timings"),
+			FSlateIcon(),
+			FUIAction(
+				FExecuteAction::CreateSP(this, &SBlueprintProfilerView::OnHeatMapDisplayModeChanged, EBlueprintProfilerHeatMapDisplayMode::Inclusive),
+				FCanExecuteAction(),
+				FIsActionChecked::CreateSP(this, &SBlueprintProfilerView::IsHeatMapDisplayModeSelected, EBlueprintProfilerHeatMapDisplayMode::Inclusive)
+			),
+			NAME_None,
+			EUserInterfaceActionType::Check);
+		HeatMapDisplayModeComboContent.AddMenuEntry(LOCTEXT("HeatMapDisplayMode_MaxTiming", "Max Time"),
+			LOCTEXT("HeatMapDisplayMode_MaxTiming_Desc", "Heat Map - Max Time"),
+			FSlateIcon(),
+			FUIAction(
+				FExecuteAction::CreateSP(this, &SBlueprintProfilerView::OnHeatMapDisplayModeChanged, EBlueprintProfilerHeatMapDisplayMode::MaxTiming),
+				FCanExecuteAction(),
+				FIsActionChecked::CreateSP(this, &SBlueprintProfilerView::IsHeatMapDisplayModeSelected, EBlueprintProfilerHeatMapDisplayMode::MaxTiming)
+			),
+			NAME_None,
+			EUserInterfaceActionType::Check);
+
 		ChildSlot
 		[
 			SNew(SVerticalBox)
@@ -100,6 +142,30 @@ void SBlueprintProfilerView::UpdateActiveProfilerWidget()
 					SNew(SHorizontalBox)
 					+SHorizontalBox::Slot()
 					.AutoWidth()
+					.Padding(FMargin(5, 0))
+					[
+						SAssignNew(HeatMapDisplayModeComboButton, SComboButton)
+						.ForegroundColor(this, &SBlueprintProfilerView::GetHeatMapDisplayModeButtonForegroundColor)
+						.ToolTipText(LOCTEXT("BlueprintProfilerHeatMapDisplayMode_Tooltip", "Heat Map Display Mode"))
+						.ButtonStyle(FEditorStyle::Get(), "ToggleButton")
+						.ContentPadding(2)
+						.MenuContent()
+						[
+							HeatMapDisplayModeComboContent.MakeWidget()
+						]
+						.ButtonContent()
+						[
+							CreateHeatMapDisplayModeButton()
+						]
+					]
+					+SHorizontalBox::Slot()
+					.Padding(FMargin(5, 0))
+					[
+						SNew(SSeparator)
+						.Orientation(Orient_Vertical)
+					]
+					+SHorizontalBox::Slot()
+					.AutoWidth()
 					.Padding(FMargin(5,0))
 					[
 						SAssignNew(ViewComboButton, SComboButton)
@@ -177,6 +243,52 @@ void SBlueprintProfilerView::OnViewSelectionChanged(const EBlueprintPerfViewType
 	}
 }
 
+FText SBlueprintProfilerView::GetCurrentHeatMapDisplayModeText() const
+{
+	IBlueprintProfilerInterface& ProfilerModule = FModuleManager::LoadModuleChecked<IBlueprintProfilerInterface>("BlueprintProfiler");
+	switch (ProfilerModule.GetGraphNodeHeatMapDisplayMode())
+	{
+	case EBlueprintProfilerHeatMapDisplayMode::None:		return LOCTEXT("HeatMapDisplayModeLabel_None", "Heat Map: Off");
+	case EBlueprintProfilerHeatMapDisplayMode::Exclusive:	return LOCTEXT("HeatMapDisplayModeLabel_Exclusive", "Heat Map: Exclusive");
+	case EBlueprintProfilerHeatMapDisplayMode::Inclusive:	return LOCTEXT("HeatMapDisplayModeLabel_Inclusive", "Heat Map: Inclusive");
+	case EBlueprintProfilerHeatMapDisplayMode::MaxTiming:	return LOCTEXT("HeatMapDisplayModeLabel_MaxTiming", "Heat Map: Max Timing");
+	}
+	return LOCTEXT("HeatMapDisplayModeLabel_Unknown", "<unknown>");
+}
+
+TSharedRef<SWidget> SBlueprintProfilerView::CreateHeatMapDisplayModeButton() const
+{
+	return SNew(SHorizontalBox)
+		+ SHorizontalBox::Slot()
+		.VAlign(VAlign_Center)
+		[
+			SNew(STextBlock)
+			.Text(this, &SBlueprintProfilerView::GetCurrentHeatMapDisplayModeText)
+		];
+}
+
+FSlateColor SBlueprintProfilerView::GetHeatMapDisplayModeButtonForegroundColor() const
+{
+	static const FName InvertedForegroundName("InvertedForeground");
+	static const FName DefaultForegroundName("DefaultForeground");
+	return HeatMapDisplayModeComboButton->IsHovered() ? FEditorStyle::GetSlateColor(InvertedForegroundName) : FEditorStyle::GetSlateColor(DefaultForegroundName);
+}
+
+bool SBlueprintProfilerView::IsHeatMapDisplayModeSelected(const EBlueprintProfilerHeatMapDisplayMode::Type InHeatMapDisplayMode) const
+{
+	IBlueprintProfilerInterface& ProfilerModule = FModuleManager::LoadModuleChecked<IBlueprintProfilerInterface>("BlueprintProfiler");
+	return ProfilerModule.GetGraphNodeHeatMapDisplayMode() == InHeatMapDisplayMode;
+}
+
+void SBlueprintProfilerView::OnHeatMapDisplayModeChanged(const EBlueprintProfilerHeatMapDisplayMode::Type NewHeatMapDisplayMode)
+{
+	IBlueprintProfilerInterface& ProfilerModule = FModuleManager::LoadModuleChecked<IBlueprintProfilerInterface>("BlueprintProfiler");
+	if (ProfilerModule.GetGraphNodeHeatMapDisplayMode() != NewHeatMapDisplayMode)
+	{
+		ProfilerModule.SetGraphNodeHeatMapDisplayMode(NewHeatMapDisplayMode);
+	}
+}
+
 TSharedRef<SWidget> SBlueprintProfilerView::CreateActiveStatisticWidget()
 {
 	// Get profiler status
diff --git a/Engine/Source/Editor/Kismet/Private/Profiler/SBlueprintProfilerView.h b/Engine/Source/Editor/Kismet/Private/Profiler/SBlueprintProfilerView.h
--- a/Engine/Source/Editor/Kismet/Private/Profiler/SBlueprintProfilerView.h
+++ b/Engine/Source/Editor/Kismet/Private/Profiler/SBlueprintProfilerView.h
@@ -2,6 +2,7 @@
 
 #pragma once
 
+#include "BlueprintProfilerModule.h"
 
 /** Blueprint performance view type */
 namespace EBlueprintPerfViewType
@@ -59,6 +60,21 @@ protected:
 	/** Called when the profiler view type is changed */
 	void OnViewSelectionChanged(const EBlueprintPerfViewType::Type NewViewType);
 
+	/** Returns the current heat map display mode label */
+	FText GetCurrentHeatMapDisplayModeText() const;
+
+	/** Constructs and returns the heat map display mode button widget */
+	TSharedRef<SWidget> CreateHeatMapDisplayModeButton() const;
+
+	/** Returns the active heat map display mode button foreground color */
+	FSlateColor GetHeatMapDisplayModeButtonForegroundColor() const;
+
+	/** Returns whether or not the given heat map display mode is selected */
+	bool IsHeatMapDisplayModeSelected(const EBlueprintProfilerHeatMapDisplayMode::Type InHeatMapDisplayMode) const;
+
+	/** Called when the heat map display mode is changed */
+	void OnHeatMapDisplayModeChanged(const EBlueprintProfilerHeatMapDisplayMode::Type NewHeatMapDisplayMode);
+
 	/** Create active statistic display widget */
 	TSharedRef<SWidget> CreateActiveStatisticWidget();
 
@@ -76,4 +92,6 @@ protected:
 	/** View combo button widget */
 	TSharedPtr<SComboButton> ViewComboButton;
 
+	/** Heat map display mode combo button widget */
+	TSharedPtr<SComboButton> HeatMapDisplayModeComboButton;
 };
diff --git a/Engine/Source/Editor/Kismet/Private/Profiler/ScriptPerfData.cpp b/Engine/Source/Editor/Kismet/Private/Profiler/ScriptPerfData.cpp
--- a/Engine/Source/Editor/Kismet/Private/Profiler/ScriptPerfData.cpp
+++ b/Engine/Source/Editor/Kismet/Private/Profiler/ScriptPerfData.cpp
@@ -89,10 +89,9 @@ void FScriptPerfData::SetNodePerformanceThreshold(const float NodePerformanceThr
 	NodePerformanceThreshold = NodePerformanceThresholdIn;
 }
 
-FSlateColor FScriptPerfData::GetNodeHeatColor() const
+float FScriptPerfData::GetNodeHeatLevel() const
 {
-	const float Value = 1.f - FMath::Min<float>(NodeTiming / NodePerformanceThreshold, 1.f);
-	return FLinearColor(1.f, Value, Value);
+	return FMath::Min<float>(NodeTiming / NodePerformanceThreshold, 1.f);
 }
 
 void FScriptPerfData::SetInclusivePerformanceThreshold(const float InclusivePerformanceThresholdIn)
@@ -100,10 +99,9 @@ void FScriptPerfData::SetInclusivePerformanceThreshold(const float InclusivePerf
 	InclusivePerformanceThreshold = InclusivePerformanceThresholdIn;
 }
 
-FSlateColor FScriptPerfData::GetInclusiveHeatColor() const
+float FScriptPerfData::GetInclusiveHeatLevel() const
 {
-	const float Value = 1.f - FMath::Min<float>(InclusiveTiming / InclusivePerformanceThreshold, 1.f);
-	return FLinearColor(1.f, Value, Value);
+	return FMath::Min<float>(InclusiveTiming / InclusivePerformanceThreshold, 1.f);
 }
 
 void FScriptPerfData::SetMaxPerformanceThreshold(const float MaxPerformanceThresholdIn)
@@ -111,10 +109,9 @@ void FScriptPerfData::SetMaxPerformanceThreshold(const float MaxPerformanceThres
 	MaxPerformanceThreshold = MaxPerformanceThresholdIn;
 }
 
-FSlateColor FScriptPerfData::GetMaxTimeHeatColor() const
+float FScriptPerfData::GetMaxTimeHeatLevel() const
 {
-	const float Value = 1.f - FMath::Min<float>(MaxTiming / MaxPerformanceThreshold, 1.f);
-	return FLinearColor(1.f, Value, Value);
+	return FMath::Min<float>(MaxTiming / MaxPerformanceThreshold, 1.f);
 }
 
 FText FScriptPerfData::GetNodeTimingText() const
diff --git a/Engine/Source/Editor/Kismet/Public/BlueprintEditorModule.h b/Engine/Source/Editor/Kismet/Public/BlueprintEditorModule.h
--- a/Engine/Source/Editor/Kismet/Public/BlueprintEditorModule.h
+++ b/Engine/Source/Editor/Kismet/Public/BlueprintEditorModule.h
@@ -179,4 +179,7 @@ private:
 
 	/** Handle to a registered LevelViewportContextMenuBlueprintExtender delegate */
 	FDelegateHandle LevelViewportContextMenuBlueprintExtenderDelegateHandle;
+
+	/** Reference to keep our custom configuration panel alive */
+	TSharedPtr<SWidget> ConfigurationPanel;
 };
diff --git a/Engine/Source/Editor/Kismet/Public/Profiler/EventExecution.h b/Engine/Source/Editor/Kismet/Public/Profiler/EventExecution.h
--- a/Engine/Source/Editor/Kismet/Public/Profiler/EventExecution.h
+++ b/Engine/Source/Editor/Kismet/Public/Profiler/EventExecution.h
@@ -26,7 +26,7 @@ namespace EScriptExecutionNodeFlags
 		TunnelEntryPin				= 0x00020000,	// Tunnel entry pin
 		TunnelExitPin				= 0x00040000,	// Tunnel exit pin
 		ReEntrantTunnelPin			= 0x00080000,	// Re-Entrant tunnel pin
-		TunnelThenPin				= 0x00100000,	// Tunnel final exit or exec then pin
+		TunnelFinalExitPin			= 0x00100000,	// Tunnel final exit or exec then pin
 		PureChain					= 0x00200000,	// Pure node call chain
 		CyclicLinkage				= 0x00400000,	// Marks execution path as cyclic.
 		InvalidTrace				= 0x00800000,	// Indicates that node doesn't contain a valid script trace.
@@ -109,6 +109,9 @@ class KISMET_API FScriptNodePerfData
 {
 public:
 
+	/** Returns a TSet containing all valid instance names */
+	void GetValidInstanceNames(TSet<FName>& ValidInstances) const;
+
 	/** Test whether or not perf data is available for the given instance/trace path */
 	bool HasPerfDataForInstanceAndTracePath(FName InstanceName, const FTracePath& TracePath) const;
 
@@ -121,6 +124,9 @@ public:
 	/** Get global blueprint perf data for the trace path */
 	TSharedPtr<FScriptPerfData> GetBlueprintPerfDataByTracePath(const FTracePath& TracePath);
 
+	/** Get global blueprint perf data for all trace paths */
+	virtual void GetBlueprintPerfDataForAllTracePaths(FScriptPerfData& OutPerfData);
+
 protected:
 
 	/** FScriptExeutionPath hash to perf data */
@@ -233,9 +239,6 @@ public:
 	/** Returns the icon for widget UI */
 	const FSlateBrush* GetIcon() const { return Icon; }
 
-	/** Returns the profiler heat color */
-	FLinearColor GetHeatColor(const FTracePath& TracePath) const;
-
 	/** Returns the current expansion state for widget UI */
 	bool IsExpanded() const { return bExpansionState; }
 
@@ -252,7 +255,7 @@ public:
 	void SetPureNodeScriptCodeRange(FInt32Range InScriptCodeRange) { PureNodeScriptCodeRange = InScriptCodeRange; }
 
 	/** Return the linear execution path from this node */
-	void GetLinearExecutionPath(TArray<FLinearExecPath>& LinearExecutionNodes, const FTracePath& TracePath);
+	virtual void GetLinearExecutionPath(TArray<FLinearExecPath>& LinearExecutionNodes, const FTracePath& TracePath);
 
 	/** Get statistic container type */
 	virtual EScriptStatContainerType::Type GetStatisticContainerType() const;
@@ -298,6 +301,117 @@ protected:
 	FInt32Range PureNodeScriptCodeRange;
 };
 
+//////////////////////////////////////////////////////////////////////////
+// FScriptExecutionTunnelInstance
+
+class KISMET_API FScriptExecutionTunnelInstance : public FScriptExecutionNode
+{
+public:
+
+	FScriptExecutionTunnelInstance(const FScriptExecNodeParams& InitParams)
+		: FScriptExecutionNode(InitParams)
+	{
+	}
+
+	// FScriptNodePerfData
+	virtual void GetBlueprintPerfDataForAllTracePaths(FScriptPerfData& OutPerfData) override;
+	// ~FScriptNodePerfData
+
+	/** Adds a custom child entry point */
+	void AddCustomEntryPoint(TSharedPtr<class FScriptExecutionTunnelEntry> CustomEntryPoint) { CustomEntryPoints.Add(CustomEntryPoint); }
+
+private:
+
+	/** List of custom entry points */
+	TArray<TSharedPtr<class FScriptExecutionTunnelEntry>> CustomEntryPoints;
+
+};
+
+//////////////////////////////////////////////////////////////////////////
+// FScriptExecutionTunnelEntry
+
+class KISMET_API FScriptExecutionTunnelEntry : public FScriptExecutionNode
+{
+public:
+
+	FScriptExecutionTunnelEntry(const FScriptExecNodeParams& InitParams, const class UK2Node_Tunnel* TunnelInstanceIn)
+		: FScriptExecutionNode(InitParams)
+		, TunnelInstance(TunnelInstanceIn)
+	{
+	}
+
+	// FScriptExecutionNode
+	virtual void GetLinearExecutionPath(TArray<FLinearExecPath>& LinearExecutionNodes, const FTracePath& TracePath) override;
+	virtual void RefreshStats(const FTracePath& TracePath) override;
+	// ~FScriptExecutionNode
+
+	/** Add tunnel timing */
+	void AddTunnelTiming(const FName InstanceName, const FTracePath& TracePath, const double Time);
+
+	/** Update tunnel exit site */
+	void UpdateTunnelExit(const FName InstanceName, const FTracePath& TracePath, const int32 ExitScriptOffset);
+
+	/** Returns true if the script offset is valid for this tunnel */
+	bool IsPathValidForTunnel(const int32 ScriptOffset) const;
+
+	/** Returns the tunnel instance graph node */
+	const UK2Node_Tunnel* GetTunnelInstance() const { return TunnelInstance.Get(); }
+
+	/** Adds a valid exit script offset */
+	void AddExitSite(TSharedPtr<FScriptExecutionNode> ExitSite, const int32 ExitScriptOffset);
+
+	/** Returns the exit point associated with the script offset */
+	TSharedPtr<FScriptExecutionNode> GetExitSite(const int32 ScriptOffset) const;
+
+	/** Returns if the exit point associated with the script offset is the final exit site */
+	bool IsFinalExitSite(const int32 ScriptOffset) const;
+
+private:
+
+	/** The tunnel instance this node represents */
+	TWeakObjectPtr<const UK2Node_Tunnel> TunnelInstance;
+	/** List of valid exit scripts offsets */
+	TMap<int32, TSharedPtr<FScriptExecutionNode>> ValidExitPoints;
+
+};
+
+//////////////////////////////////////////////////////////////////////////
+// FScriptExecutionTunnelExit
+
+class KISMET_API FScriptExecutionTunnelExit : public FScriptExecutionNode
+{
+public:
+
+	FScriptExecutionTunnelExit(const FScriptExecNodeParams& InitParams)
+		: FScriptExecutionNode(InitParams)
+	{
+	}
+
+	// FScriptExecutionNode
+	virtual void GetLinearExecutionPath(TArray<FLinearExecPath>& LinearExecutionNodes, const FTracePath& TracePath) override;
+	// ~FScriptExecutionNode
+
+};
+
+//////////////////////////////////////////////////////////////////////////
+// FScriptExecutionPureNode
+
+class KISMET_API FScriptExecutionPureNode : public FScriptExecutionNode
+{
+public:
+
+	FScriptExecutionPureNode(const FScriptExecNodeParams& InitParams)
+		: FScriptExecutionNode(InitParams)
+	{
+	}
+
+	// FScriptExecutionNode
+	virtual void GetLinearExecutionPath(TArray<FLinearExecPath>& LinearExecutionNodes, const FTracePath& TracePath) override;
+	virtual void RefreshStats(const FTracePath& TracePath) override;
+	// ~FScriptExecutionNode
+
+};
+
 //////////////////////////////////////////////////////////////////////////
 // FScriptExecutionInstance
 
diff --git a/Engine/Source/Editor/Kismet/Public/Profiler/ScriptPerfData.h b/Engine/Source/Editor/Kismet/Public/Profiler/ScriptPerfData.h
--- a/Engine/Source/Editor/Kismet/Public/Profiler/ScriptPerfData.h
+++ b/Engine/Source/Editor/Kismet/Public/Profiler/ScriptPerfData.h
@@ -34,6 +34,9 @@ public:
 	/** Add branch data */
 	void AddBranchData(const FScriptPerfData& DataIn);
 
+	/** Increments sample count without affecting data */
+	void TickSamples() { NumSamples++; }
+
 	/** Reset the current data buffer and all derived data stats */
 	void Reset();
 
@@ -55,10 +58,10 @@ public:
 	double GetTotalTiming() const { return TotalTiming; }
 	int32 GetSampleCount() const { return NumSamples; }
 
-	// Returns various performance colors for visual display
-	FSlateColor GetNodeHeatColor() const;
-	FSlateColor GetInclusiveHeatColor() const;
-	FSlateColor GetMaxTimeHeatColor() const;
+	// Returns various performance heat levels for visual display
+	float GetNodeHeatLevel() const;
+	float GetInclusiveHeatLevel() const;
+	float GetMaxTimeHeatLevel() const;
 
 	// Returns the various stats this container holds in FText format
 	FText GetTotalTimingText() const;
diff --git a/Engine/Source/Editor/Kismet/Public/Profiler/TracePath.h b/Engine/Source/Editor/Kismet/Public/Profiler/TracePath.h
--- a/Engine/Source/Editor/Kismet/Public/Profiler/TracePath.h
+++ b/Engine/Source/Editor/Kismet/Public/Profiler/TracePath.h
@@ -18,6 +18,16 @@ public:
 
 	FTracePath(const FTracePath& TracePathIn)
 		: CachedCrc(TracePathIn.CachedCrc)
+		, Tunnel(TracePathIn.Tunnel)
+#if TRACEPATH_DEBUG
+		, TraceStack(TracePathIn.TraceStack)
+#endif
+	{
+	}
+
+	FTracePath(const FTracePath& TracePathIn, TSharedPtr<const class FScriptExecutionTunnelEntry> TunnelNode)
+		: CachedCrc(TracePathIn.CachedCrc)
+		, Tunnel(TunnelNode)
 #if TRACEPATH_DEBUG
 		, TraceStack(TracePathIn.TraceStack)
 #endif
@@ -30,6 +40,7 @@ public:
 	{
 		TraceStack = TracePathIn.TraceStack;
 		CachedCrc = TracePathIn.CachedCrc;
+		Tunnel = TracePathIn.Tunnel;
 		return *this;
 	}
 #endif
@@ -53,9 +64,12 @@ public:
 		TraceStack.SetNum(0);
 #endif
 		CachedCrc = 0;
-
+		Tunnel.Reset();
 	}
 
+	/** Returns the current tunnel */
+	TSharedPtr<const class FScriptExecutionTunnelEntry> GetTunnel() const { return Tunnel; }
+
 #if TRACEPATH_DEBUG
 	const FString GetPathString() const
 	{
@@ -71,10 +85,13 @@ public:
 
 private:
 
-#if TRACEPATH_DEBUG
-	TArray<int32> TraceStack;
-#endif
 	/** Cached crc */
 	uint32 CachedCrc;
+	/** The current tunnel entry point */
+	TSharedPtr<const class FScriptExecutionTunnelEntry> Tunnel;
 
+#if TRACEPATH_DEBUG
+	/** Contributing pin list - for debugging only */
+	TArray<int32> TraceStack;
+#endif
 };
diff --git a/Engine/Source/Editor/KismetCompiler/Private/KismetCompilerMisc.cpp b/Engine/Source/Editor/KismetCompiler/Private/KismetCompilerMisc.cpp
--- a/Engine/Source/Editor/KismetCompiler/Private/KismetCompilerMisc.cpp
+++ b/Engine/Source/Editor/KismetCompiler/Private/KismetCompilerMisc.cpp
@@ -1482,12 +1482,6 @@ void FKismetFunctionContext::SetExternalNetNameMap(FNetNameMapping* NewMap)
 	NetNameMap = NewMap;
 }
 
-int32 FKismetFunctionContext::GetContextUniqueID()
-{
-	static int32 UUIDCounter = 1;
-	return UUIDCounter++;
-}
-
 bool FKismetFunctionContext::DoesStatementRequiresSwitch(const FBlueprintCompiledStatement* Statement)
 {
 	return Statement && (
diff --git a/Engine/Source/Editor/KismetCompiler/Private/KismetCompilerModule.cpp b/Engine/Source/Editor/KismetCompiler/Private/KismetCompilerModule.cpp
--- a/Engine/Source/Editor/KismetCompiler/Private/KismetCompilerModule.cpp
+++ b/Engine/Source/Editor/KismetCompiler/Private/KismetCompilerModule.cpp
@@ -152,7 +152,6 @@ extern UNREALED_API FSecondsCounterData BlueprintCompileAndLoadTimerData;
 // Compiles a blueprint.
 void FKismet2CompilerModule::CompileBlueprint(class UBlueprint* Blueprint, const FKismetCompilerOptions& CompileOptions, FCompilerResultsLog& Results, TSharedPtr<FBlueprintCompileReinstancer> ParentReinstancer, TArray<UObject*>* ObjLoaded)
 {
-	SCOPE_SECONDS_COUNTER(GBlueprintCompileTime);
 	FSecondsCounterScope Timer(BlueprintCompileAndLoadTimerData);
 	BP_SCOPED_COMPILER_EVENT_STAT(EKismetCompilerStats_CompileTime);
 
diff --git a/Engine/Source/Editor/KismetCompiler/Private/KismetCompilerVMBackend.cpp b/Engine/Source/Editor/KismetCompiler/Private/KismetCompilerVMBackend.cpp
--- a/Engine/Source/Editor/KismetCompiler/Private/KismetCompilerVMBackend.cpp
+++ b/Engine/Source/Editor/KismetCompiler/Private/KismetCompilerVMBackend.cpp
@@ -1403,7 +1403,7 @@ public:
 		}
 		else
 		{
-			int32 EventType = 0;
+			uint8 EventType = 0;
 			switch (Statement.Type)
 			{
 			case KCST_InstrumentedWireExit:			EventType = EScriptInstrumentation::NodeExit; break;
@@ -1417,8 +1417,6 @@ public:
 			}
 			Writer << EX_InstrumentationEvent;
 			Writer << EventType;
-
-			Offset += sizeof(int32);
 		}
 
 		if (Statement.Type != KCST_InstrumentedStateRestore)
@@ -1453,7 +1451,8 @@ public:
 					}
 
 					TArray<TWeakObjectPtr<UEdGraphNode>> MacroInstanceNodes;
-					bool bBreakpointSite = Statement.Type == KCST_DebugSite;
+					const bool bInstrumentedBreakpoint = Statement.Type == KCST_InstrumentedWireEntry;
+					bool bBreakpointSite = Statement.Type == KCST_DebugSite || bInstrumentedBreakpoint;
 
 					if (MacroSourceNode)
 					{
diff --git a/Engine/Source/Editor/KismetCompiler/Public/KismetCompiledFunctionContext.h b/Engine/Source/Editor/KismetCompiler/Public/KismetCompiledFunctionContext.h
--- a/Engine/Source/Editor/KismetCompiler/Public/KismetCompiledFunctionContext.h
+++ b/Engine/Source/Editor/KismetCompiler/Public/KismetCompiledFunctionContext.h
@@ -227,8 +227,6 @@ public:
 		}
 	}
 
-	KISMETCOMPILER_API int32 GetContextUniqueID();
-
 	UBlueprint* GetBlueprint()
 	{
 		return Blueprint;
diff --git a/Engine/Source/Editor/UnrealEd/Private/Kismet2/BlueprintEditorUtils.cpp b/Engine/Source/Editor/UnrealEd/Private/Kismet2/BlueprintEditorUtils.cpp
--- a/Engine/Source/Editor/UnrealEd/Private/Kismet2/BlueprintEditorUtils.cpp
+++ b/Engine/Source/Editor/UnrealEd/Private/Kismet2/BlueprintEditorUtils.cpp
@@ -2262,10 +2262,7 @@ void FBlueprintEditorUtils::MarkBlueprintAsModified(UBlueprint* Blueprint)
 			}
 		}
 
-		if (Blueprint->BlueprintType != BPTYPE_MacroLibrary)
-		{
-			Blueprint->Status = BS_Dirty;
-		}
+		Blueprint->Status = BS_Dirty;
 		Blueprint->MarkPackageDirty();
 		Blueprint->PostEditChange();
 	}
@@ -2401,10 +2398,14 @@ UFunction* FBlueprintEditorUtils::FindFunctionInImplementedInterfaces(const UBlu
 			UClass* SearchClass = (*It);
 			if( SearchClass )
 			{
-				if (UFunction* OverriddenFunction = SearchClass->FindFunctionByName(FunctionName, EIncludeSuperFlag::ExcludeSuper))
+				do
 				{
-					return OverriddenFunction;
-				}
+					if( UFunction* OverriddenFunction = SearchClass->FindFunctionByName(FunctionName, EIncludeSuperFlag::ExcludeSuper) )
+					{
+						return OverriddenFunction;
+					}
+					SearchClass = SearchClass->GetSuperClass();
+				} while (SearchClass);
 			}
 			else if( bOutInvalidInterface )
 			{
@@ -4615,15 +4616,7 @@ void FBlueprintEditorUtils::ChangeMemberVariableType(UBlueprint* Blueprint, cons
 							BlueprintEditor->SummonSearchUI(bSetFindWithinBlueprint, AllVariableNodes[0]->GetFindReferenceSearchString(), bSelectFirstResult);
 						}
 					}
-					else
-					{
-						// And recompile
-						FBlueprintEditorUtils::MarkBlueprintAsStructurallyModified(Blueprint);
-					}
 				}
-
-				// And recompile
-				FBlueprintEditorUtils::MarkBlueprintAsStructurallyModified(Blueprint);
 			}
 		}
 	}
@@ -5083,9 +5076,6 @@ void FBlueprintEditorUtils::ChangeLocalVariableType(UBlueprint* InBlueprint, con
 					K2Schema->ReconstructNode(*VariableNode, true);
 				}
 
-				// And recompile
-				FBlueprintEditorUtils::MarkBlueprintAsStructurallyModified(InBlueprint);
-
 				TSharedPtr<IToolkit> FoundAssetEditor = FToolkitManager::Get().FindEditorForAsset(InBlueprint);
 
 				// No need to submit a search query if there are no nodes.
@@ -5889,7 +5879,23 @@ void FBlueprintEditorUtils::ConformImplementedEvents(UBlueprint* Blueprint)
 				// If the event is loaded and is not a custom event
 				if(!EventNode->HasAnyFlags(RF_NeedLoad|RF_NeedPostLoad) && EventNode->bOverrideFunction)
 				{
-					const bool bEventNodeUsedByInterface = ImplementedInterfaceClasses.Contains(EventNode->EventReference.GetMemberParentClass(EventNode->GetBlueprintClassFromNode()));
+					UClass* EventClass = EventNode->EventReference.GetMemberParentClass(EventNode->GetBlueprintClassFromNode());
+					bool bEventNodeUsedByInterface = false;
+					int32 Idx = 0;
+					while (Idx != ImplementedInterfaceClasses.Num() && bEventNodeUsedByInterface == false)
+					{
+						const UClass* CurrentInterface = ImplementedInterfaceClasses[Idx];
+						while (CurrentInterface)
+						{
+							if (EventClass == CurrentInterface )
+							{
+								bEventNodeUsedByInterface = true;
+								break;
+							}
+							CurrentInterface = CurrentInterface->GetSuperClass();
+						}
+						++Idx;
+					}
 					if (Blueprint->GeneratedClass && !bEventNodeUsedByInterface)
 					{
 						FixOverriddenEventSignature(EventNode, Blueprint, CurrentGraph);
diff --git a/Engine/Source/Editor/UnrealEd/Private/Kismet2/KismetDebugUtilities.cpp b/Engine/Source/Editor/UnrealEd/Private/Kismet2/KismetDebugUtilities.cpp
--- a/Engine/Source/Editor/UnrealEd/Private/Kismet2/KismetDebugUtilities.cpp
+++ b/Engine/Source/Editor/UnrealEd/Private/Kismet2/KismetDebugUtilities.cpp
@@ -762,7 +762,15 @@ void FKismetDebugUtilities::SetBreakpointInternal(UBreakpoint* Breakpoint, bool
 	{
 		if (uint8* InstallSite = InstallSites[i])
 		{
-			*InstallSite = bShouldBeEnabled ? EX_Breakpoint : EX_Tracepoint;
+			if (*InstallSite == EX_InstrumentationEvent)
+			{
+				InstallSite++;
+				*InstallSite = bShouldBeEnabled ? EScriptInstrumentation::NodeDebugSite : EScriptInstrumentation::NodeEntry;
+			}
+			else
+			{
+				*InstallSite = bShouldBeEnabled ? EX_Breakpoint : EX_Tracepoint;
+			}
 		}
 	}
 }
diff --git a/Engine/Source/Editor/UnrealEd/Private/ScriptDisassembler.cpp b/Engine/Source/Editor/UnrealEd/Private/ScriptDisassembler.cpp
--- a/Engine/Source/Editor/UnrealEd/Private/ScriptDisassembler.cpp
+++ b/Engine/Source/Editor/UnrealEd/Private/ScriptDisassembler.cpp
@@ -869,12 +869,15 @@ void FKismetBytecodeDisassembler::ProcessCommon(int32& ScriptIndex, EExprToken O
 		}
 	case EX_InstrumentationEvent:
 		{
-			const int32 EventType = ReadINT(ScriptIndex);
+			const uint8 EventType = ReadBYTE(ScriptIndex);
 			switch (EventType)
 			{
 				case EScriptInstrumentation::PureNodeEntry:
 					Ar.Logf(TEXT("%s $%X: .. instrumented pure node entry site .."), *Indents, (int32)Opcode);
 					break;
+				case EScriptInstrumentation::NodeDebugSite:
+					Ar.Logf(TEXT("%s $%X: .. instrumented debug site .."), *Indents, (int32)Opcode);
+					break;
 				case EScriptInstrumentation::NodeEntry:
 					Ar.Logf(TEXT("%s $%X: .. instrumented wire entry site .."), *Indents, (int32)Opcode);
 					break;
@@ -887,6 +890,12 @@ void FKismetBytecodeDisassembler::ProcessCommon(int32& ScriptIndex, EExprToken O
 				case EScriptInstrumentation::RestoreState:
 					Ar.Logf(TEXT("%s $%X: .. restore execution state .."), *Indents, (int32)Opcode);
 					break;
+				case EScriptInstrumentation::ResetState:
+					Ar.Logf(TEXT("%s $%X: .. reset execution state .."), *Indents, (int32)Opcode);
+					break;
+				case EScriptInstrumentation::SuspendState:
+					Ar.Logf(TEXT("%s $%X: .. suspend execution state .."), *Indents, (int32)Opcode);
+					break;
 				case EScriptInstrumentation::PopState:
 					Ar.Logf(TEXT("%s $%X: .. pop execution state .."), *Indents, (int32)Opcode);
 					break;
diff --git a/Engine/Source/Programs/AutomationTool/Scripts/CookCommand.Automation.cs b/Engine/Source/Programs/AutomationTool/Scripts/CookCommand.Automation.cs
--- a/Engine/Source/Programs/AutomationTool/Scripts/CookCommand.Automation.cs
+++ b/Engine/Source/Programs/AutomationTool/Scripts/CookCommand.Automation.cs
@@ -223,7 +223,7 @@ public partial class Project : CommandUtils
                     {
                         // If you change this target path you must also update logic in CookOnTheFlyServer.cpp. Passing a single directory around is cumbersome for testing, so I have hard coded it.
                         // Similarly if you change the .uplugin name you must update DefaultPluginName in BlueprintNativeCodeGenModule.cpp
-                        string GeneratedPluginPath = CombinePaths(GetDirectoryName(ProjectDir), "Intermediate", Platform, "NativizedAssets/NativizedAssets.uplugin");
+                        string GeneratedPluginPath = CombinePaths(ProjectDir, "Intermediate", Platform, "NativizedAssets/NativizedAssets.uplugin");
                         Params.BlueprintPluginPaths.Add( new FileReference(GeneratedPluginPath) );
                     }
                 }
diff --git a/Engine/Source/Runtime/BlueprintRuntime/BlueprintRuntime.Build.cs b/Engine/Source/Runtime/BlueprintRuntime/BlueprintRuntime.Build.cs
new file mode 100644
--- /dev/null
+++ b/Engine/Source/Runtime/BlueprintRuntime/BlueprintRuntime.Build.cs
@@ -0,0 +1,25 @@
+// Copyright 1998-2016 Epic Games, Inc. All Rights Reserved.
+
+namespace UnrealBuildTool.Rules
+{
+	public class BlueprintRuntime : ModuleRules
+	{
+        public BlueprintRuntime(TargetInfo Target)
+		{
+			PrivateIncludePaths.AddRange(
+				new string[] {
+					"Runtime/BlueprintRuntime/Private",
+				}
+			);
+
+			PrivateDependencyModuleNames.AddRange(
+				new string[]
+				{
+					"Core",
+				    "CoreUObject",
+                    "Settings",
+				}
+			);
+		}
+	}
+}
\ No newline at end of file
diff --git a/Engine/Source/Runtime/BlueprintRuntime/Private/BlueprintRuntime.cpp b/Engine/Source/Runtime/BlueprintRuntime/Private/BlueprintRuntime.cpp
new file mode 100644
--- /dev/null
+++ b/Engine/Source/Runtime/BlueprintRuntime/Private/BlueprintRuntime.cpp
@@ -0,0 +1,52 @@
+// Copyright 1998-2016 Epic Games, Inc. All Rights Reserved.
+
+#include "BlueprintRuntimePrivatePCH.h"
+#include "ISettingsModule.h"
+#include "BlueprintRuntimeSettings.h"
+#include "BlueprintSupport.h"
+
+#define LOCTEXT_NAMESPACE "BlueprintRuntime"
+
+class FBlueprintRuntime : public IBlueprintRuntime
+{
+	/** IModuleInterface implementation */
+	virtual void StartupModule() override;
+	virtual void ShutdownModule() override;
+	virtual void PropagateWarningSettings() override;
+};
+IMPLEMENT_MODULE(FBlueprintRuntime, BlueprintRuntime)
+
+void FBlueprintRuntime::StartupModule()
+{
+}
+
+void FBlueprintRuntime::ShutdownModule()
+{
+	// This function may be called during shutdown to clean up your module.  For modules that support dynamic reloading,
+	// we call this function before unloading the module.
+}
+
+void FBlueprintRuntime::PropagateWarningSettings()
+{
+	const UBlueprintRuntimeSettings* BlueprintRuntimeSettings = GetDefault<UBlueprintRuntimeSettings>();
+	// Propagate current settings to the core blueprint runtime, which can not easily
+	// leverage the reflection system:
+	TArray<FName> WarningsToTreatAsErrors;
+	TArray<FName> WarningsToSuppress;
+
+	for (const FBlueprintWarningSettings& Setting : BlueprintRuntimeSettings->WarningSettings)
+	{
+		if(Setting.WarningBehavior == EBlueprintWarningBehavior::Error)
+		{
+			WarningsToTreatAsErrors.Add(Setting.WarningIdentifier);
+		}
+		else if( Setting.WarningBehavior == EBlueprintWarningBehavior::Suppress)
+		{
+			WarningsToSuppress.Add(Setting.WarningIdentifier);
+		}
+	}
+	FBlueprintSupport::UpdateWarningBehavior(WarningsToTreatAsErrors, WarningsToSuppress);
+}
+
+
+#undef LOCTEXT_NAMESPACE
diff --git a/Engine/Source/Runtime/BlueprintRuntime/Private/BlueprintRuntimePrivatePCH.h b/Engine/Source/Runtime/BlueprintRuntime/Private/BlueprintRuntimePrivatePCH.h
new file mode 100644
--- /dev/null
+++ b/Engine/Source/Runtime/BlueprintRuntime/Private/BlueprintRuntimePrivatePCH.h
@@ -0,0 +1,6 @@
+// Copyright 1998-2016 Epic Games, Inc. All Rights Reserved.
+
+#include "BlueprintRuntime.h"
+#include "CoreUObject.h"
+
+
diff --git a/Engine/Source/Runtime/BlueprintRuntime/Private/BlueprintRuntimeSettings.cpp b/Engine/Source/Runtime/BlueprintRuntime/Private/BlueprintRuntimeSettings.cpp
new file mode 100644
--- /dev/null
+++ b/Engine/Source/Runtime/BlueprintRuntime/Private/BlueprintRuntimeSettings.cpp
@@ -0,0 +1,12 @@
+// Copyright 1998-2016 Epic Games, Inc. All Rights Reserved.
+
+#include "BlueprintRuntimePrivatePCH.h"
+#include "BlueprintRuntimeSettings.h"
+
+//////////////////////////////////////////////////////////////////////////
+// UBlueprintRuntimeSettings
+
+UBlueprintRuntimeSettings::UBlueprintRuntimeSettings(const FObjectInitializer& ObjectInitializer)
+: Super(ObjectInitializer)
+{
+}
diff --git a/Engine/Source/Runtime/BlueprintRuntime/Public/BlueprintRuntime.h b/Engine/Source/Runtime/BlueprintRuntime/Public/BlueprintRuntime.h
new file mode 100644
--- /dev/null
+++ b/Engine/Source/Runtime/BlueprintRuntime/Public/BlueprintRuntime.h
@@ -0,0 +1,26 @@
+// Copyright 1998-2016 Epic Games, Inc. All Rights Reserved.
+
+#pragma once
+
+#include "ModuleManager.h"
+
+/**
+ * The public interface to this module
+ */
+class IBlueprintRuntime : public IModuleInterface
+{
+public:
+	/**
+	 * Singleton-like access to this module's interface.  This is just for convenience!
+	 * Beware of calling this during the shutdown phase, though.  Your module might have been unloaded already.
+	 *
+	 * @return Returns singleton instance, loading the module on demand if needed
+	 */
+	static inline IBlueprintRuntime& Get()
+	{
+		return FModuleManager::LoadModuleChecked< IBlueprintRuntime >("BlueprintRuntime");
+	}
+
+	virtual void PropagateWarningSettings() = 0;
+};
+
diff --git a/Engine/Source/Runtime/BlueprintRuntime/Public/BlueprintRuntimeSettings.h b/Engine/Source/Runtime/BlueprintRuntime/Public/BlueprintRuntimeSettings.h
new file mode 100644
--- /dev/null
+++ b/Engine/Source/Runtime/BlueprintRuntime/Public/BlueprintRuntimeSettings.h
@@ -0,0 +1,40 @@
+// Copyright 1998-2016 Epic Games, Inc. All Rights Reserved.
+
+#pragma once
+
+#include "BlueprintRuntimeSettings.generated.h"
+
+UENUM()
+enum class EBlueprintWarningBehavior : uint8
+{
+	Warn,
+	Error,
+	Suppress,
+};
+
+USTRUCT()
+struct FBlueprintWarningSettings
+{
+	GENERATED_USTRUCT_BODY()
+
+	UPROPERTY()
+	FName WarningIdentifier;
+
+	UPROPERTY()
+	FText WarningDescription;
+
+	UPROPERTY()
+	EBlueprintWarningBehavior WarningBehavior;
+};
+
+/**
+* Implements the settings for the BlueprintRuntime module
+*/
+UCLASS(config = Engine, defaultconfig)
+class BLUEPRINTRUNTIME_API UBlueprintRuntimeSettings : public UObject
+{
+	GENERATED_UCLASS_BODY()
+
+	UPROPERTY(EditAnywhere, config, Category = Warnings)
+	TArray<FBlueprintWarningSettings> WarningSettings;
+};
diff --git a/Engine/Source/Runtime/Core/Private/Async/TaskGraph.cpp b/Engine/Source/Runtime/Core/Private/Async/TaskGraph.cpp
--- a/Engine/Source/Runtime/Core/Private/Async/TaskGraph.cpp
+++ b/Engine/Source/Runtime/Core/Private/Async/TaskGraph.cpp
@@ -1462,7 +1462,7 @@ public:
 				Name = FString::Printf(TEXT("TaskGraphThreadNP %d"), ThreadIndex - (LastExternalThread + 1));
 				ThreadPri = TPri_BelowNormal; // we want normal tasks below normal threads like the game thread
 			}
-			uint32 StackSize = 256 * 1024;
+			uint32 StackSize = 384 * 1024;
 			WorkerThreads[ThreadIndex].RunnableThread = FRunnableThread::Create(&Thread(ThreadIndex), *Name, StackSize, ThreadPri, FPlatformAffinity::GetTaskGraphThreadMask()); // these are below normal threads so that they sleep when the named threads are active
 			WorkerThreads[ThreadIndex].bAttached = true;
 		}
diff --git a/Engine/Source/Runtime/Core/Private/Misc/Core.cpp b/Engine/Source/Runtime/Core/Private/Misc/Core.cpp
--- a/Engine/Source/Runtime/Core/Private/Misc/Core.cpp
+++ b/Engine/Source/Runtime/Core/Private/Misc/Core.cpp
@@ -259,9 +259,6 @@ bool					GPumpingMessagesOutsideOfMainLoop = false;
 /** Enables various editor and HMD hacks that allow the experimental VR editor feature to work, perhaps at the expense of other systems */
 bool					GEnableVREditorHacks = false;
 
-/** Total blueprint compile time.																			*/
-double GBlueprintCompileTime = 0.0;
-
 DEFINE_STAT(STAT_AudioMemory);
 DEFINE_STAT(STAT_TextureMemory);
 DEFINE_STAT(STAT_MemoryPhysXTotalAllocationSize);
diff --git a/Engine/Source/Runtime/Core/Public/CoreGlobals.h b/Engine/Source/Runtime/Core/Public/CoreGlobals.h
--- a/Engine/Source/Runtime/Core/Public/CoreGlobals.h
+++ b/Engine/Source/Runtime/Core/Public/CoreGlobals.h
@@ -313,9 +313,6 @@ extern CORE_API float GAsyncIOBandwidthLimit;
 /** Whether or not messages are being pumped outside of main loop */
 extern CORE_API bool GPumpingMessagesOutsideOfMainLoop;
 
-/** Total blueprint compile time. */
-extern CORE_API double GBlueprintCompileTime;
-
 /** Enables various editor and HMD hacks that allow the experimental VR editor feature to work, perhaps at the expense of other systems */
 extern CORE_API bool GEnableVREditorHacks;
 
diff --git a/Engine/Source/Runtime/CoreUObject/Private/Blueprint/BlueprintSupport.cpp b/Engine/Source/Runtime/CoreUObject/Private/Blueprint/BlueprintSupport.cpp
--- a/Engine/Source/Runtime/CoreUObject/Private/Blueprint/BlueprintSupport.cpp
+++ b/Engine/Source/Runtime/CoreUObject/Private/Blueprint/BlueprintSupport.cpp
@@ -123,6 +123,36 @@ bool FBlueprintSupport::IsInBlueprintPackage(UObject* LoadedObj)
 	return false;
 }
 
+static TArray<FBlueprintWarningDeclaration> BlueprintWarnings;
+static TSet<FName> BlueprintWarningsToTreatAsError;
+static TSet<FName> BlueprintWarningsToSuppress;
+
+void FBlueprintSupport::RegisterBlueprintWarning(const FBlueprintWarningDeclaration& Warning)
+{
+	BlueprintWarnings.Add(Warning);
+}
+
+const TArray<FBlueprintWarningDeclaration>& FBlueprintSupport::GetBlueprintWarnings()
+{
+	return BlueprintWarnings;
+}
+
+void FBlueprintSupport::UpdateWarningBehavior(const TArray<FName>& WarningIdentifiersToTreatAsError, const TArray<FName>& WarningIdentifiersToSuppress)
+{
+	BlueprintWarningsToTreatAsError = TSet<FName>(WarningIdentifiersToTreatAsError);
+	BlueprintWarningsToSuppress = TSet<FName>(WarningIdentifiersToSuppress);
+}
+
+bool FBlueprintSupport::ShouldTreatWarningAsError(FName WarningIdentifier)
+{
+	return BlueprintWarningsToTreatAsError.Find(WarningIdentifier) != nullptr;
+}
+
+bool FBlueprintSupport::ShouldSuppressWarning(FName WarningIdentifier)
+{
+	return BlueprintWarningsToSuppress.Find(WarningIdentifier) != nullptr;
+}
+
 /*******************************************************************************
  * FScopedClassDependencyGather
  ******************************************************************************/
diff --git a/Engine/Source/Runtime/CoreUObject/Private/Serialization/AsyncLoading.cpp b/Engine/Source/Runtime/CoreUObject/Private/Serialization/AsyncLoading.cpp
--- a/Engine/Source/Runtime/CoreUObject/Private/Serialization/AsyncLoading.cpp
+++ b/Engine/Source/Runtime/CoreUObject/Private/Serialization/AsyncLoading.cpp
@@ -2260,6 +2260,12 @@ EAsyncPackageState::Type FAsyncPackage::PostLoadDeferredObjects(double InTickSta
 		for (UObject* Object : DeferredFinalizeObjects)
 		{
 			Object->AtomicallyClearInternalFlags(EInternalObjectFlags::AsyncLoading);
+			// Dynamic Class doesn't require/use pre-loading (or post-loading). 
+			// The CDO is created at this point, because now it's safe to solve cyclic dependencies.
+			if (UDynamicClass* DynamicClass = Cast<UDynamicClass>(Object))
+			{
+				DynamicClass->GetDefaultObject(true);
+			}
 		}
 
 		// Mark package as having been fully loaded and update load time.
@@ -2271,16 +2277,6 @@ EAsyncPackageState::Type FAsyncPackage::PostLoadDeferredObjects(double InTickSta
 
 			if (Linker)
 			{
-				// Dynamic Class doesn't require/use pre-loading (or post-loading). 
-				// The CDO is created at this point, because now it's safe to solve cyclic dependencies.
-				for (FObjectExport& Export : Linker->ExportMap)
-				{
-					if (UDynamicClass* DynamicClass = Cast<UDynamicClass>(Export.Object))
-					{
-						DynamicClass->GetDefaultObject(true);
-					}
-				}
-
 				CreateClustersFromPackage(Linker);
 
 				// give a hint to the IO system that we are done with this file for now
diff --git a/Engine/Source/Runtime/CoreUObject/Private/UObject/ScriptCore.cpp b/Engine/Source/Runtime/CoreUObject/Private/UObject/ScriptCore.cpp
--- a/Engine/Source/Runtime/CoreUObject/Private/UObject/ScriptCore.cpp
+++ b/Engine/Source/Runtime/CoreUObject/Private/UObject/ScriptCore.cpp
@@ -339,15 +339,26 @@ FString FFrame::GetScriptCallstack()
 // Error or warning handler.
 //
 //@TODO: This function should take more information in, or be able to gather it from the callstack!
-void FFrame::KismetExecutionMessage(const TCHAR* Message, ELogVerbosity::Type Verbosity)
+void FFrame::KismetExecutionMessage(const TCHAR* Message, ELogVerbosity::Type Verbosity, FName WarningId)
 {
-
 #if !UE_BUILD_SHIPPING
 	// Optionally always treat errors/warnings as bad
 	if (Verbosity <= ELogVerbosity::Warning && FParse::Param(FCommandLine::Get(), TEXT("FATALSCRIPTWARNINGS")))
 	{
 		Verbosity = ELogVerbosity::Fatal;
 	}
+	else if(Verbosity == ELogVerbosity::Warning && WarningId != FName())
+	{
+		// check to see if this specific warning has been elevated to an error:
+		if( FBlueprintSupport::ShouldTreatWarningAsError(WarningId) )
+		{
+			Verbosity = ELogVerbosity::Error;
+		}
+		else if(FBlueprintSupport::ShouldSuppressWarning(WarningId))
+		{
+			return;
+		}
+	}
 #endif
 
 	FString ScriptStack;
@@ -1534,7 +1545,7 @@ IMPLEMENT_VM_FUNCTION( EX_WireTracepoint, execWireTracepoint );
 void UObject::execInstrumentation( FFrame& Stack, RESULT_DECL )
 {
 #if !UE_BUILD_SHIPPING
-	const EScriptInstrumentation::Type EventType = static_cast<EScriptInstrumentation::Type>(Stack.ReadInt<int32>());
+	const EScriptInstrumentation::Type EventType = static_cast<EScriptInstrumentation::Type>(Stack.PeekCode());
 #if WITH_EDITORONLY_DATA
 	if (GIsEditor)
 	{
@@ -1548,10 +1559,16 @@ void UObject::execInstrumentation( FFrame& Stack, RESULT_DECL )
 			FBlueprintExceptionInfo WiretraceExceptionInfo(EBlueprintExceptionType::WireTracepoint);
 			FBlueprintCoreDelegates::ThrowScriptException(this, Stack, WiretraceExceptionInfo);
 		}
+		else if (EventType == EScriptInstrumentation::NodeDebugSite)
+		{
+			FBlueprintExceptionInfo TracepointExceptionInfo(EBlueprintExceptionType::Breakpoint);
+			FBlueprintCoreDelegates::ThrowScriptException(this, Stack, TracepointExceptionInfo);
+		}
 	}
 #endif
 	EScriptInstrumentationEvent InstrumentationEventInfo(EventType, this, Stack);
 	FBlueprintCoreDelegates::InstrumentScriptEvent(InstrumentationEventInfo);
+	Stack.SkipCode(1);
 #endif
 }
 IMPLEMENT_VM_FUNCTION( EX_InstrumentationEvent, execInstrumentation );
diff --git a/Engine/Source/Runtime/CoreUObject/Public/Blueprint/BlueprintSupport.h b/Engine/Source/Runtime/CoreUObject/Public/Blueprint/BlueprintSupport.h
--- a/Engine/Source/Runtime/CoreUObject/Public/Blueprint/BlueprintSupport.h
+++ b/Engine/Source/Runtime/CoreUObject/Public/Blueprint/BlueprintSupport.h
@@ -4,6 +4,18 @@
 
 #include "GCObject.h"
 
+struct FBlueprintWarningDeclaration
+{
+	FBlueprintWarningDeclaration(FName InWarningIdentifier, FText InWarningDescription)
+		: WarningIdentifier(InWarningIdentifier)
+		, WarningDescription( InWarningDescription )
+	{
+	}
+
+	FName WarningIdentifier;
+	FText WarningDescription;
+};
+
 /** 
  * This set of functions contains blueprint related UObject functionality.
  */
@@ -33,6 +45,12 @@ struct FBlueprintSupport
 
 	/** Not a particularly fast function. Mostly intended for validation in debug builds. */
 	static bool IsInBlueprintPackage(UObject* LoadedObj);
+
+	COREUOBJECT_API static void RegisterBlueprintWarning(const FBlueprintWarningDeclaration& Warning);
+	COREUOBJECT_API static const TArray<FBlueprintWarningDeclaration>& GetBlueprintWarnings();
+	COREUOBJECT_API static void UpdateWarningBehavior(const TArray<FName>& WarningIdentifiersToTreatAsError, const TArray<FName>& WarningIdentifiersToSuppress);
+	COREUOBJECT_API static bool ShouldTreatWarningAsError(FName WarningIdentifier);
+	COREUOBJECT_API static bool ShouldSuppressWarning(FName WarningIdentifier);
 };
 
 #if WITH_EDITOR
diff --git a/Engine/Source/Runtime/CoreUObject/Public/UObject/Script.h b/Engine/Source/Runtime/CoreUObject/Public/UObject/Script.h
--- a/Engine/Source/Runtime/CoreUObject/Public/UObject/Script.h
+++ b/Engine/Source/Runtime/CoreUObject/Public/UObject/Script.h
@@ -306,6 +306,7 @@ namespace EScriptInstrumentation
 		Event,
 		ResumeEvent,
 		PureNodeEntry,
+		NodeDebugSite,
 		NodeEntry,
 		NodeExit,
 		PushState,
diff --git a/Engine/Source/Runtime/CoreUObject/Public/UObject/ScriptSerialization.h b/Engine/Source/Runtime/CoreUObject/Public/UObject/ScriptSerialization.h
--- a/Engine/Source/Runtime/CoreUObject/Public/UObject/ScriptSerialization.h
+++ b/Engine/Source/Runtime/CoreUObject/Public/UObject/ScriptSerialization.h
@@ -260,7 +260,7 @@
 		}
 		case EX_InstrumentationEvent:
 		{
-			iCode += sizeof(int32);
+			iCode += sizeof(uint8);
 			break;
 		}
 		case EX_Return:
diff --git a/Engine/Source/Runtime/CoreUObject/Public/UObject/Stack.h b/Engine/Source/Runtime/CoreUObject/Public/UObject/Stack.h
--- a/Engine/Source/Runtime/CoreUObject/Public/UObject/Stack.h
+++ b/Engine/Source/Runtime/CoreUObject/Public/UObject/Stack.h
@@ -134,7 +134,13 @@ public:
 
 	COREUOBJECT_API virtual void Serialize( const TCHAR* V, ELogVerbosity::Type Verbosity, const class FName& Category ) override;
 	
-	COREUOBJECT_API static void KismetExecutionMessage(const TCHAR* Message, ELogVerbosity::Type Verbosity);
+	COREUOBJECT_API static void KismetExecutionMessage(const TCHAR* Message, ELogVerbosity::Type Verbosity, FName WarningId = FName());
+
+	/** Returns the current script op code */
+	const uint8 PeekCode() const { return *Code; }
+
+	/** Skips over the number of op codes specified by NumOps */
+	void SkipCode(const int32 NumOps) { Code += NumOps; }
 
 	template<typename TNumericType>
 	TNumericType ReadInt();
diff --git a/Engine/Source/Runtime/Engine/Private/Blueprint.cpp b/Engine/Source/Runtime/Engine/Private/Blueprint.cpp
--- a/Engine/Source/Runtime/Engine/Private/Blueprint.cpp
+++ b/Engine/Source/Runtime/Engine/Private/Blueprint.cpp
@@ -1225,13 +1225,13 @@ void UBlueprint::TagSubobjects(EObjectFlags NewFlags)
 {
 	Super::TagSubobjects(NewFlags);
 
-	if (GeneratedClass && !GeneratedClass->HasAnyFlags(GARBAGE_COLLECTION_KEEPFLAGS) && !GeneratedClass->IsRooted())
+	if (GeneratedClass && !GeneratedClass->HasAnyFlags(GARBAGE_COLLECTION_KEEPFLAGS))
 	{
 		GeneratedClass->SetFlags(NewFlags);
 		GeneratedClass->TagSubobjects(NewFlags);
 	}
 
-	if (SkeletonGeneratedClass && SkeletonGeneratedClass != GeneratedClass && !SkeletonGeneratedClass->HasAnyFlags(GARBAGE_COLLECTION_KEEPFLAGS) && !SkeletonGeneratedClass->IsRooted())
+	if (SkeletonGeneratedClass && SkeletonGeneratedClass != GeneratedClass && !SkeletonGeneratedClass->HasAnyFlags(GARBAGE_COLLECTION_KEEPFLAGS))
 	{
 		SkeletonGeneratedClass->SetFlags(NewFlags);
 		SkeletonGeneratedClass->TagSubobjects(NewFlags);
diff --git a/Engine/Source/Runtime/Engine/Private/BlueprintGeneratedClass.cpp b/Engine/Source/Runtime/Engine/Private/BlueprintGeneratedClass.cpp
--- a/Engine/Source/Runtime/Engine/Private/BlueprintGeneratedClass.cpp
+++ b/Engine/Source/Runtime/Engine/Private/BlueprintGeneratedClass.cpp
@@ -227,6 +227,12 @@ void UBlueprintGeneratedClass::ConditionalRecompileClass(TArray<UObject*>* ObjLo
 				// Make sure that nodes are up to date, so that we get any updated blueprint signatures
 				FBlueprintEditorUtils::RefreshExternalBlueprintDependencyNodes(GeneratingBP);
 
+				// Normal blueprints get their status reset by RecompileBlueprintBytecode, but macros will not:
+				if ((GeneratingBP->Status != BS_Error) && (GeneratingBP->BlueprintType == EBlueprintType::BPTYPE_MacroLibrary))
+				{
+					GeneratingBP->Status = BS_UpToDate;
+				}
+
 				if (Package != nullptr && Package->IsDirty() && !bStartedWithUnsavedChanges)
 				{
 					Package->SetDirtyFlag(false);
diff --git a/Engine/Source/Runtime/Engine/Private/KismetArrayLibrary.cpp b/Engine/Source/Runtime/Engine/Private/KismetArrayLibrary.cpp
--- a/Engine/Source/Runtime/Engine/Private/KismetArrayLibrary.cpp
+++ b/Engine/Source/Runtime/Engine/Private/KismetArrayLibrary.cpp
@@ -1,16 +1,55 @@
 // Copyright 1998-2016 Epic Games, Inc. All Rights Reserved.
 
 #include "EnginePrivate.h"
+#include "Blueprint/BlueprintSupport.h"
 #include "Kismet/KismetArrayLibrary.h"
 
 DEFINE_LOG_CATEGORY_STATIC(LogArray, Warning, All);
 
+#define LOCTEXT_NAMESPACE "UKismetArrayLibrary"
+
 //////////////////////////////////////////////////////////////////////////
 // UKismetArrayLibrary
 
+const FName GetOutOfBoundsWarning = FName("GetOutOfBoundsWarning");
+const FName SetOutOfBoundsWarning = FName("SetOutOfBoundsWarning");
+const FName InsertOutOfBoundsWarning = FName("InsertOutOfBoundsWarning");
+const FName RemoveOutOfBoundsWarning = FName("RemoveOutOfBoundsWarning");
+const FName ResizeArrayNegativeWarning = FName("ResizeArrayNegativeWarning");
+
 UKismetArrayLibrary::UKismetArrayLibrary(const FObjectInitializer& ObjectInitializer)
 	: Super(ObjectInitializer)
 {
+	FBlueprintSupport::RegisterBlueprintWarning(
+		FBlueprintWarningDeclaration (
+			GetOutOfBoundsWarning,
+			LOCTEXT("GetOutOfBoundsWarning", "Array read access out of bounds")
+		)
+	);
+	FBlueprintSupport::RegisterBlueprintWarning(
+		FBlueprintWarningDeclaration (
+			SetOutOfBoundsWarning,
+			LOCTEXT("SetOutOfBoundsWarning", "Array assignment access out of bounds")
+		)
+	);
+	FBlueprintSupport::RegisterBlueprintWarning(
+		FBlueprintWarningDeclaration (
+			InsertOutOfBoundsWarning,
+			LOCTEXT("InsertOutOfBoundsWarning", "Array insertion out of bounds")
+		)
+	);
+	FBlueprintSupport::RegisterBlueprintWarning(
+		FBlueprintWarningDeclaration (
+			RemoveOutOfBoundsWarning,
+			LOCTEXT("RemoveOutOfBoundsWarning", "Array removal out of bounds")
+		)
+	);
+	FBlueprintSupport::RegisterBlueprintWarning(
+		FBlueprintWarningDeclaration (
+			ResizeArrayNegativeWarning,
+			LOCTEXT("ResizeArrayNegativeWarning", "Array resized to negative size")
+		)
+	);
 }
 
 void UKismetArrayLibrary::FilterArray(const TArray<AActor*>& TargetArray, TSubclassOf<class AActor> FilterClass, TArray<AActor*>& FilteredArray)
@@ -92,7 +131,7 @@ void UKismetArrayLibrary::GenericArray_Insert(void* TargetArray, const UArrayPro
 		}
 		else
 		{
-			FFrame::KismetExecutionMessage(*FString::Printf(TEXT("Attempted to insert an item into array %s out of bounds [%d/%d]!"), *ArrayProp->GetName(), Index, GetLastIndex(ArrayHelper)), ELogVerbosity::Warning);
+			FFrame::KismetExecutionMessage(*FString::Printf(TEXT("Attempted to insert an item into array %s out of bounds [%d/%d]!"), *ArrayProp->GetName(), Index, GetLastIndex(ArrayHelper)), ELogVerbosity::Warning, InsertOutOfBoundsWarning);
 		}
 	}
 }
@@ -108,7 +147,7 @@ void UKismetArrayLibrary::GenericArray_Remove(void* TargetArray, const UArrayPro
 		}
 		else
 		{
-			FFrame::KismetExecutionMessage(*FString::Printf(TEXT("Attempted to remove an item from an invalid index from array %s [%d/%d]!"), *ArrayProp->GetName(), IndexToRemove, GetLastIndex(ArrayHelper)), ELogVerbosity::Warning);
+			FFrame::KismetExecutionMessage(*FString::Printf(TEXT("Attempted to remove an item from an invalid index from array %s [%d/%d]!"), *ArrayProp->GetName(), IndexToRemove, GetLastIndex(ArrayHelper)), ELogVerbosity::Warning, RemoveOutOfBoundsWarning);
 		}
 	}
 }
@@ -172,7 +211,7 @@ void UKismetArrayLibrary::GenericArray_Resize(void* TargetArray, const UArrayPro
 		}
 		else
 		{
-			FFrame::KismetExecutionMessage(*FString::Printf(TEXT("Attempted to resize an array using negative size: Array = %s, Size = %d!"), *ArrayProp->GetName(), Size), ELogVerbosity::Warning);
+			FFrame::KismetExecutionMessage(*FString::Printf(TEXT("Attempted to resize an array using negative size: Array = %s, Size = %d!"), *ArrayProp->GetName(), Size), ELogVerbosity::Warning, ResizeArrayNegativeWarning);
 		}
 	}
 }
@@ -218,7 +257,8 @@ void UKismetArrayLibrary::GenericArray_Get(void* TargetArray, const UArrayProper
 				Index,
 				*ArrayProp->GetName(),
 				ArrayHelper.Num()),
-				ELogVerbosity::Error);
+				ELogVerbosity::Warning,
+				GetOutOfBoundsWarning);
 			InnerProp->InitializeValue(Item);
 		}
 	}
@@ -244,7 +284,7 @@ void UKismetArrayLibrary::GenericArray_Set(void* TargetArray, const UArrayProper
 		}
 		else
 		{
-			FFrame::KismetExecutionMessage(*FString::Printf(TEXT("Attempted to set an invalid index on array %s [%d/%d]!"), *ArrayProp->GetName(), Index, GetLastIndex(ArrayHelper)), ELogVerbosity::Warning);
+			FFrame::KismetExecutionMessage(*FString::Printf(TEXT("Attempted to set an invalid index on array %s [%d/%d]!"), *ArrayProp->GetName(), Index, GetLastIndex(ArrayHelper)), ELogVerbosity::Warning, SetOutOfBoundsWarning);
 		}
 	}
 }
@@ -379,4 +419,6 @@ void UKismetArrayLibrary::SetArrayPropertyByName(UObject* Object, FName Property
 {
 	// We should never hit these!  They're stubs to avoid NoExport on the class.
 	check(0);
-}
\ No newline at end of file
+}
+
+#undef LOCTEXT_NAMESPACE
diff --git a/Engine/Source/Runtime/Engine/Private/KismetMathLibrary.cpp b/Engine/Source/Runtime/Engine/Private/KismetMathLibrary.cpp
--- a/Engine/Source/Runtime/Engine/Private/KismetMathLibrary.cpp
+++ b/Engine/Source/Runtime/Engine/Private/KismetMathLibrary.cpp
@@ -1,10 +1,14 @@
 // Copyright 1998-2016 Epic Games, Inc. All Rights Reserved.
 
 #include "EnginePrivate.h"
+
+#include "Blueprint/BlueprintSupport.h"
 #include "Kismet/KismetMathLibrary.h"
 #include "CoreStats.h"
 #include "Math/DualQuat.h"
 
+#define LOCTEXT_NAMESPACE "UKismetMathLibrary"
+
 /** Interpolate a linear alpha value using an ease mode and function. */
 float EaseAlpha(float InAlpha, uint8 EasingFunc, float BlendExp, int32 Steps)
 {
@@ -27,9 +31,38 @@ float EaseAlpha(float InAlpha, uint8 EasingFunc, float BlendExp, int32 Steps)
 	return InAlpha;
 }
 
+const FName DivideByZeroWarning = FName("DivideByZeroWarning");
+const FName NegativeSqrtWarning = FName("NegativeSqrtWarning");
+const FName ZeroLengthProjectionWarning = FName("ZeroLengthProjectionWarning");
+const FName InvalidDateWarning = FName("InvalidDateWarning");
+
 UKismetMathLibrary::UKismetMathLibrary(const FObjectInitializer& ObjectInitializer)
 	: Super(ObjectInitializer)
 {
+	FBlueprintSupport::RegisterBlueprintWarning(
+		FBlueprintWarningDeclaration (
+			DivideByZeroWarning,
+			LOCTEXT("DivideByZeroWarning", "Divide by zero")
+		)
+	);
+	FBlueprintSupport::RegisterBlueprintWarning(
+		FBlueprintWarningDeclaration (
+			NegativeSqrtWarning,
+			LOCTEXT("NegativeSqrtWarning", "Square root of negative number")
+		)
+	);
+	FBlueprintSupport::RegisterBlueprintWarning(
+		FBlueprintWarningDeclaration (
+			ZeroLengthProjectionWarning,
+			LOCTEXT("ZeroLengthProjectionWarning", "Projection onto vector of zero length")
+		)
+	);
+	FBlueprintSupport::RegisterBlueprintWarning(
+		FBlueprintWarningDeclaration (
+			InvalidDateWarning,
+			LOCTEXT("InvalidDateWarning", "Invalid date warning")
+		)
+	);
 }
 
 bool UKismetMathLibrary::RandomBool()
@@ -86,8 +119,7 @@ uint8 UKismetMathLibrary::Divide_ByteByte(uint8 A, uint8 B)
 {
 	if (B == 0)
 	{
-		//@TODO: EXCEPTION: Throw script exception 
-		FFrame::KismetExecutionMessage(TEXT("Divide by zero: Divide_ByteByte"), ELogVerbosity::Warning);
+		FFrame::KismetExecutionMessage(TEXT("Divide by zero: Divide_ByteByte"), ELogVerbosity::Warning, DivideByZeroWarning);
 	}
 
 	return (B != 0) ? (A / B) : 0;
@@ -97,8 +129,7 @@ uint8 UKismetMathLibrary::Percent_ByteByte(uint8 A, uint8 B)
 {
 	if (B == 0)
 	{
-		//@TODO: EXCEPTION: Whatever on these sites
-		FFrame::KismetExecutionMessage(TEXT("Modulo by zero"), ELogVerbosity::Warning);
+		FFrame::KismetExecutionMessage(TEXT("Modulo by zero"), ELogVerbosity::Warning, DivideByZeroWarning);
 	}
 
 	return (B != 0) ? (A % B) : 0;
@@ -163,8 +194,7 @@ int32 UKismetMathLibrary::Divide_IntInt(int32 A, int32 B)
 {
 	if (B == 0)
 	{
-		//@TODO: EXCEPTION: Throw script exception 
-		FFrame::KismetExecutionMessage(TEXT("Divide by zero: Divide_IntInt"), ELogVerbosity::Warning);
+		FFrame::KismetExecutionMessage(TEXT("Divide by zero: Divide_IntInt"), ELogVerbosity::Warning, DivideByZeroWarning);
 	}
 
 	return (B != 0) ? (A / B) : 0;
@@ -174,8 +204,7 @@ int32 UKismetMathLibrary::Percent_IntInt(int32 A, int32 B)
 {
 	if (B == 0)
 	{
-		//@TODO: EXCEPTION: Throw script exception 
-		FFrame::KismetExecutionMessage(TEXT("Modulo by zero"), ELogVerbosity::Warning);
+		FFrame::KismetExecutionMessage(TEXT("Modulo by zero"), ELogVerbosity::Warning, DivideByZeroWarning);
 	}
 
 	return (B != 0) ? (A % B) : 0;
@@ -509,7 +538,7 @@ float UKismetMathLibrary::Sqrt(float A)
 	}
 	else if (A < 0.f)
 	{
-		FFrame::KismetExecutionMessage(TEXT("Attempt to take Sqrt() of negative number - returning 0."), ELogVerbosity::Warning);
+		FFrame::KismetExecutionMessage(TEXT("Attempt to take Sqrt() of negative number - returning 0."), ELogVerbosity::Warning, NegativeSqrtWarning);
 	}
 
 	return Result;
@@ -551,8 +580,7 @@ int32 UKismetMathLibrary::FMod(float Dividend, float Divisor, float& Remainder)
 	}
 	else
 	{
-		//@TODO: EXCEPTION: Throw script exception 
-		FFrame::KismetExecutionMessage(TEXT("Attempted modulo 0 - returning 0."), ELogVerbosity::Warning);
+		FFrame::KismetExecutionMessage(TEXT("Attempted modulo 0 - returning 0."), ELogVerbosity::Warning, DivideByZeroWarning);
 
 		Result = 0;
 		Remainder = 0.f;
@@ -723,8 +751,7 @@ FVector UKismetMathLibrary::Divide_VectorFloat(FVector A, float B)
 {
 	if (B == 0.f)
 	{
-		//@TODO: EXCEPTION: Throw script exception 
-		FFrame::KismetExecutionMessage(TEXT("Divide by zero: Divide_VectorFloat"), ELogVerbosity::Warning);
+		FFrame::KismetExecutionMessage(TEXT("Divide by zero: Divide_VectorFloat"), ELogVerbosity::Warning, DivideByZeroWarning);
 		return FVector::ZeroVector;
 	}
 
@@ -735,8 +762,7 @@ FVector UKismetMathLibrary::Divide_VectorInt(FVector A, int32 B)
 {
 	if (B == 0)
 	{
-		//@TODO: EXCEPTION: Throw script exception 
-		FFrame::KismetExecutionMessage(TEXT("Divide by zero: Divide_VectorInt"), ELogVerbosity::Warning);
+		FFrame::KismetExecutionMessage(TEXT("Divide by zero: Divide_VectorInt"), ELogVerbosity::Warning, DivideByZeroWarning);
 		return FVector::ZeroVector;
 	}
 
@@ -747,8 +773,7 @@ FVector UKismetMathLibrary::Divide_VectorVector(FVector A, FVector B)
 {
 	if (B.X == 0.f || B.Y == 0.f || B.Z == 0.f)
 	{
-		//@TODO: EXCEPTION: Throw script exception 
-		FFrame::KismetExecutionMessage(TEXT("Divide by zero: Divide_VectorVector"), ELogVerbosity::Warning);
+		FFrame::KismetExecutionMessage(TEXT("Divide by zero: Divide_VectorVector"), ELogVerbosity::Warning, DivideByZeroWarning);
 		return FVector::ZeroVector;
 	}
 
@@ -942,7 +967,7 @@ FVector UKismetMathLibrary::ProjectVectorOnToVector(FVector V, FVector Target)
 	}
 	else
 	{
-		FFrame::KismetExecutionMessage(TEXT("Divide by zero: ProjectVectorOnToVector with zero Target vector"), ELogVerbosity::Warning);
+		FFrame::KismetExecutionMessage(TEXT("Divide by zero: ProjectVectorOnToVector with zero Target vector"), ELogVerbosity::Warning, ZeroLengthProjectionWarning);
 		return FVector::ZeroVector;
 	}
 }
@@ -1252,8 +1277,7 @@ FVector2D UKismetMathLibrary::Divide_Vector2DFloat(FVector2D A, float B)
 {
 	if (B == 0.f)
 	{
-		//@TODO: EXCEPTION: Throw script exception 
-		FFrame::KismetExecutionMessage(TEXT("Divide by zero: Divide_Vector2DFloat"), ELogVerbosity::Warning);
+		FFrame::KismetExecutionMessage(TEXT("Divide by zero: Divide_Vector2DFloat"), ELogVerbosity::Warning, DivideByZeroWarning);
 		return FVector2D::ZeroVector;
 	}
 
@@ -1328,7 +1352,7 @@ FDateTime UKismetMathLibrary::MakeDateTime(int32 Year, int32 Month, int32 Day, i
 {
 	if (!FDateTime::Validate(Year, Month, Day, Hour, Minute, Second, Millisecond))
 	{
-		FFrame::KismetExecutionMessage(*FString::Printf(TEXT("DateTime in bad format (year %d, month %d, day %d, hour %d, minute %d, second %d, millisecond %d). E.g. year, month and day can't be zero."), Year, Month, Day, Hour, Minute, Second, Millisecond), ELogVerbosity::Warning);
+		FFrame::KismetExecutionMessage(*FString::Printf(TEXT("DateTime in bad format (year %d, month %d, day %d, hour %d, minute %d, second %d, millisecond %d). E.g. year, month and day can't be zero."), Year, Month, Day, Hour, Minute, Second, Millisecond), ELogVerbosity::Warning, InvalidDateWarning);
 
 		return FDateTime(1, 1, 1, 0, 0, 0, 0);
 	}
@@ -1481,8 +1505,7 @@ int32 UKismetMathLibrary::DaysInMonth( int32 Year, int32 Month )
 {
 	if ((Month < 1) || (Month > 12))
 	{
-		//@TODO: EXCEPTION: Throw script exception
-		FFrame::KismetExecutionMessage(TEXT("Invalid month (must be between 1 and 12): DaysInMonth"), ELogVerbosity::Warning);
+		FFrame::KismetExecutionMessage(TEXT("Invalid month (must be between 1 and 12): DaysInMonth"), ELogVerbosity::Warning, InvalidDateWarning);
 		return 0;
 	}
 
@@ -2315,3 +2338,4 @@ FRandomStream UKismetMathLibrary::MakeRandomStream(int32 InitialSeed)
 	return FRandomStream(InitialSeed);
 }
 
+#undef LOCTEXT_NAMESPACE
diff --git a/Engine/Source/Runtime/Engine/Public/GeneratedCodeHelpers.h b/Engine/Source/Runtime/Engine/Public/GeneratedCodeHelpers.h
--- a/Engine/Source/Runtime/Engine/Public/GeneratedCodeHelpers.h
+++ b/Engine/Source/Runtime/Engine/Public/GeneratedCodeHelpers.h
@@ -39,8 +39,7 @@ inline FBox2D CreateFBox2D(FVector2D InMin, FVector2D InMax, bool InIsValid)
 template<class NativeType>
 inline NativeType* NoNativeCast(UClass* NoNativeClass, UObject* Object)
 {
-	check(NoNativeClass);
-	check(!Object || (nullptr != Cast<NativeType>(Object)));
+	check(NoNativeClass && NoNativeClass->IsChildOf<NativeType>());
 	return (Object && Object->IsA(NoNativeClass)) ? ((NativeType*)Object) : nullptr;
 }
 
diff --git a/Engine/Source/Runtime/GameplayAbilities/Private/Abilities/GameplayAbility.cpp b/Engine/Source/Runtime/GameplayAbilities/Private/Abilities/GameplayAbility.cpp
--- a/Engine/Source/Runtime/GameplayAbilities/Private/Abilities/GameplayAbility.cpp
+++ b/Engine/Source/Runtime/GameplayAbilities/Private/Abilities/GameplayAbility.cpp
@@ -19,15 +19,21 @@
 UGameplayAbility::UGameplayAbility(const FObjectInitializer& ObjectInitializer)
 	: Super(ObjectInitializer)
 {
+	auto ImplementedInBlueprint = [](const UFunction* Func) -> bool
+	{
+		return Func && ensure(Func->GetOuter())
+			&& (Func->GetOuter()->IsA(UBlueprintGeneratedClass::StaticClass()) || Func->GetOuter()->IsA(UDynamicClass::StaticClass()));
+	};
+
 	{
 		static FName FuncName = FName(TEXT("K2_ShouldAbilityRespondToEvent"));
 		UFunction* ShouldRespondFunction = GetClass()->FindFunctionByName(FuncName);
-		bHasBlueprintShouldAbilityRespondToEvent = ShouldRespondFunction && ensure(ShouldRespondFunction->GetOuter()) && ShouldRespondFunction->GetOuter()->IsA(UBlueprintGeneratedClass::StaticClass());
+		bHasBlueprintShouldAbilityRespondToEvent = ImplementedInBlueprint(ShouldRespondFunction);
 	}
 	{
 		static FName FuncName = FName(TEXT("K2_CanActivateAbility"));
 		UFunction* CanActivateFunction = GetClass()->FindFunctionByName(FuncName);
-		bHasBlueprintCanUse = CanActivateFunction && ensure(CanActivateFunction->GetOuter()) && CanActivateFunction->GetOuter()->IsA(UBlueprintGeneratedClass::StaticClass());
+		bHasBlueprintCanUse = ImplementedInBlueprint(CanActivateFunction);
 	}
 	{
 		static FName FuncName = FName(TEXT("K2_ActivateAbility"));
@@ -35,13 +41,13 @@ UGameplayAbility::UGameplayAbility(const FObjectInitializer& ObjectInitializer)
 		// FIXME: temp to work around crash
 		if (ActivateFunction && (HasAnyFlags(RF_ClassDefaultObject) || ActivateFunction->IsValidLowLevelFast()))
 		{
-			bHasBlueprintActivate = ActivateFunction && ensure(ActivateFunction->GetOuter()) && ActivateFunction->GetOuter()->IsA(UBlueprintGeneratedClass::StaticClass());
+			bHasBlueprintActivate = ImplementedInBlueprint(ActivateFunction);
 		}
 	}
 	{
 		static FName FuncName = FName(TEXT("K2_ActivateAbilityFromEvent"));
 		UFunction* ActivateFunction = GetClass()->FindFunctionByName(FuncName);
-		bHasBlueprintActivateFromEvent = ActivateFunction && ensure(ActivateFunction->GetOuter()) && ActivateFunction->GetOuter()->IsA(UBlueprintGeneratedClass::StaticClass());
+		bHasBlueprintActivateFromEvent = ImplementedInBlueprint(ActivateFunction);
 	}
 	
 #if WITH_EDITOR
diff --git a/Engine/Source/Runtime/Launch/Private/Launch.cpp b/Engine/Source/Runtime/Launch/Private/Launch.cpp
--- a/Engine/Source/Runtime/Launch/Private/Launch.cpp
+++ b/Engine/Source/Runtime/Launch/Private/Launch.cpp
@@ -72,6 +72,9 @@ void LaunchStaticShutdownAfterError()
 	TermGamePhys();
 }
 
+#if WITH_EDITOR
+extern UNREALED_API FSecondsCounterData BlueprintCompileAndLoadTimerData;
+#endif
 
 /**
  * Static guarded main function. Rolled into own function so we can have error handling for debug/ release builds depending
@@ -142,7 +145,12 @@ int32 GuardedMain( const TCHAR* CmdLine )
 	}
 
 	double EngineInitializationTime = FPlatformTime::Seconds() - GStartTime;
-	UE_LOG(LogLoad, Log, TEXT("Full Startup: %.2f seconds (BP compile: %.2f seconds)"), EngineInitializationTime, GBlueprintCompileTime);
+	UE_LOG(LogLoad, Log, TEXT("(Engine Initialization) Total time: %.2f seconds"), EngineInitializationTime);
+
+#if WITH_EDITOR
+	UE_LOG(LogLoad, Log, TEXT("(Engine Initialization) Total Blueprint compile time: %.2f seconds"), BlueprintCompileAndLoadTimerData.GetTime());
+#endif
+
 	ACCUM_LOADTIME(TEXT("EngineInitialization"), EngineInitializationTime);
 
 	while( !GIsRequestingExit )
