Copying //UE4/Dev-Blueprints to //UE4/Dev-Main (Source: //UE4/Dev-Blueprints @ 3025888)

diff --git a/Engine/Source/Developer/BlueprintProfiler/Private/ScriptInstrumentationPlayback.cpp b/Engine/Source/Developer/BlueprintProfiler/Private/ScriptInstrumentationPlayback.cpp
--- a/Engine/Source/Developer/BlueprintProfiler/Private/ScriptInstrumentationPlayback.cpp
+++ b/Engine/Source/Developer/BlueprintProfiler/Private/ScriptInstrumentationPlayback.cpp
@@ -3,6 +3,7 @@
 #include "BlueprintProfilerPCH.h"
 #include "EditorStyleSet.h"
 #include "GraphEditorSettings.h"
+#include "BlueprintEditorUtils.h"
 
 #define LOCTEXT_NAMESPACE "ScriptInstrumentationPlayback"
 
@@ -39,18 +40,79 @@ bool FBlueprintExecutionContext::InitialiseContext(const FString& BlueprintPath)
 		BlueprintParams.IconColor = FLinearColor(0.46f, 0.54f, 0.81f);
 		BlueprintNode = MakeShareable<FScriptExecutionBlueprint>(new FScriptExecutionBlueprint(BlueprintParams));
 		// Map the blueprint execution
-		if (MapBlueprintExecution())
-		{
-			return true;
-		}
+		bIsBlueprintMapped = MapBlueprintExecution();
 	}
-	return false;
+	return bIsBlueprintMapped;
+}
+
+void FBlueprintExecutionContext::RemoveMapping()
+{
+	bIsBlueprintMapped = false;
+	BlueprintClass.Reset();
+	BlueprintNode.Reset();
+	FunctionContexts.Reset();
+	UbergraphFunctionName = NAME_None;
 }
 
-void FBlueprintExecutionContext::AddEventNode(TSharedPtr<FScriptExecutionNode> EventExecNode)
+bool FBlueprintExecutionContext::IsEventMapped(const FName EventName) const
+{
+	return EventFunctionContexts.Contains(EventName);
+}
+
+void FBlueprintExecutionContext::AddEventNode(TSharedPtr<FBlueprintFunctionContext> FunctionContext, TSharedPtr<FScriptExecutionNode> EventExecNode)
 {
 	check(BlueprintNode.IsValid());
 	BlueprintNode->AddChildNode(EventExecNode);
+	EventFunctionContexts.Add(EventExecNode->GetName()) = FunctionContext;
+}
+
+void FBlueprintExecutionContext::RegisterEventContext(const FName EventName, TSharedPtr<FBlueprintFunctionContext> FunctionContext)
+{
+	EventFunctionContexts.Add(EventName) = FunctionContext;
+}
+
+FName FBlueprintExecutionContext::MapBlueprintInstance(const FString& InstancePath)
+{
+	FName InstanceName(*InstancePath);
+	TWeakObjectPtr<const UObject> Instance;
+	const bool bNewInstance = ResolveInstance(InstanceName, Instance);
+	const bool bBlueprintMatches = Instance.IsValid() ? (Instance->GetClass() == BlueprintClass) : false;
+	if (bBlueprintMatches)
+	{
+		if (bNewInstance)
+		{
+			// Create new instance node
+			FScriptExecNodeParams InstanceNodeParams;
+			InstanceNodeParams.NodeName = InstanceName;
+			InstanceNodeParams.ObservedObject = Instance.Get();
+			InstanceNodeParams.NodeFlags = EScriptExecutionNodeFlags::Instance;
+			const AActor* Actor = Cast<AActor>(Instance.Get());
+			InstanceNodeParams.DisplayName = Actor ? FText::FromString(Actor->GetActorLabel()) : FText::FromString(Instance.Get()->GetName());
+			InstanceNodeParams.Tooltip = LOCTEXT("NavigateToInstanceHyperlink_ToolTip", "Navigate to the Instance");
+			InstanceNodeParams.IconColor = FLinearColor(1.f, 1.f, 1.f, 1.f);
+			InstanceNodeParams.Icon = const_cast<FSlateBrush*>(FEditorStyle::GetBrush(TEXT("BlueprintProfiler.Actor")));
+			TSharedPtr<FScriptExecutionInstance> InstanceNode = MakeShareable<FScriptExecutionInstance>(new FScriptExecutionInstance(InstanceNodeParams));
+			// Link to parent blueprint entry
+			BlueprintNode->AddInstance(InstanceNode);
+			// Fill out events from the blueprint root node
+			for (int32 NodeIdx = 0; NodeIdx < BlueprintNode->GetNumChildren(); ++NodeIdx)
+			{
+				InstanceNode->AddChildNode(BlueprintNode->GetChildByIndex(NodeIdx));
+			}
+			// Broadcast change
+			IBlueprintProfilerInterface& ProfilerModule = FModuleManager::LoadModuleChecked<IBlueprintProfilerInterface>("BlueprintProfiler");
+			ProfilerModule.GetGraphLayoutChangedDelegate().Broadcast(Blueprint.Get());
+		}
+		else
+		{
+			TSharedPtr<FScriptExecutionInstance> InstanceNode = StaticCastSharedPtr<FScriptExecutionInstance>(BlueprintNode->GetInstanceByName(InstanceName));
+			if (InstanceNode.IsValid() && PIEActorInstances.Contains(InstanceName))
+			{
+				InstanceNode->SetPIEInstance(PIEActorInstances.Find(InstanceName)->Get());
+			}
+		}
+	}
+	return InstanceName;
 }
 
 bool FBlueprintExecutionContext::HasProfilerDataForInstance(const FName InstanceName) const
@@ -75,18 +137,18 @@ FName FBlueprintExecutionContext::RemapInstancePath(const FName InstanceName) co
 	return Result;
 }
 
-TWeakObjectPtr<const UObject> FBlueprintExecutionContext::GetInstance(FName& InstanceNameInOut)
+bool FBlueprintExecutionContext::ResolveInstance(FName& InstanceNameInOut, TWeakObjectPtr<const UObject>& ObjectInOut)
 {
-	TWeakObjectPtr<const UObject> Result;
+	bool bNewInstance = false;
 	FName CorrectedName = InstanceNameInOut;
 	if (const FName* SearchName = PIEInstanceNameMap.Find(InstanceNameInOut)) 
 	{
 		CorrectedName = *SearchName;
 		InstanceNameInOut = CorrectedName;
 	}
-	if (TWeakObjectPtr<const UObject>* SearchResult = ActorInstances.Find(CorrectedName))
+	if (TWeakObjectPtr<const UObject>* SearchResult = EditorActorInstances.Find(CorrectedName))
 	{
-		Result = *SearchResult;
+		ObjectInOut = *SearchResult;
 	}
 	else
 	{
@@ -105,26 +167,31 @@ TWeakObjectPtr<const UObject> FBlueprintExecutionContext::GetInstance(FName& Ins
 						{
 							if (UObject* EditorObject = FindObject<UObject>(LevelIter, *ObjectPtr->GetName()))
 							{
-								CorrectedName = FName(*EditorObject->GetPathName());
-								ActorInstances.Add(CorrectedName) = EditorObject;
-								PIEInstanceNameMap.Add(InstanceNameInOut) = CorrectedName;
-								InstanceNameInOut = CorrectedName;
-								Result = EditorObject;
+								if (EditorObject->GetClass() == BlueprintClass)
+								{
+									CorrectedName = FName(*EditorObject->GetPathName());
+									EditorActorInstances.Add(CorrectedName) = EditorObject;
+									PIEActorInstances.Add(CorrectedName) = ObjectPtr;
+									PIEInstanceNameMap.Add(InstanceNameInOut) = CorrectedName;
+									InstanceNameInOut = CorrectedName;
+									ObjectInOut = EditorObject;
+								}
 								break;
 							}
 						}
 					}
 				}
-				else if (!ObjectPtr->HasAnyFlags(RF_Transient))
+				else if (!ObjectPtr->HasAnyFlags(RF_Transient) && ObjectPtr->GetClass() == BlueprintClass)
 				{
-					ActorInstances.Add(InstanceNameInOut) = ObjectPtr;
+					EditorActorInstances.Add(InstanceNameInOut) = ObjectPtr;
 					PIEInstanceNameMap.Add(InstanceNameInOut) = InstanceNameInOut;
-					Result = ObjectPtr;
+					ObjectInOut = ObjectPtr;
+					bNewInstance = true;
 				}
 			}
 		}
 	}
-	return Result;
+	return bNewInstance;
 }
 
 FName FBlueprintExecutionContext::GetEventFunctionName(const FName EventName) const
@@ -132,6 +199,16 @@ FName FBlueprintExecutionContext::GetEventFunctionName(const FName EventName) co
 	return (EventName == UEdGraphSchema_K2::FN_UserConstructionScript) ? UEdGraphSchema_K2::FN_UserConstructionScript : UbergraphFunctionName;
 }
 
+TSharedPtr<FBlueprintFunctionContext> FBlueprintExecutionContext::GetFunctionContextForEvent(const FName EventName) const
+{
+	TSharedPtr<FBlueprintFunctionContext> Result;
+	if (const TSharedPtr<FBlueprintFunctionContext>* SearchResult = EventFunctionContexts.Find(EventName))
+	{
+		Result = *SearchResult;
+	}
+	return Result;
+}
+
 TSharedPtr<FBlueprintFunctionContext> FBlueprintExecutionContext::GetFunctionContext(const FName FunctionNameIn) const
 {
 	TSharedPtr<FBlueprintFunctionContext> Result;
@@ -144,20 +221,63 @@ TSharedPtr<FBlueprintFunctionContext> FBlueprintExecutionContext::GetFunctionCon
 	return Result;
 }
 
-TSharedPtr<FBlueprintFunctionContext> FBlueprintExecutionContext::CreateFunctionContext(UEdGraph* Graph)
+TSharedPtr<FBlueprintFunctionContext> FBlueprintExecutionContext::CreateFunctionContext(const FName FunctionName, UEdGraph* Graph)
 {
-	FName FunctionName = Graph->GetFName();
-	FunctionName = (FunctionName == UEdGraphSchema_K2::GN_EventGraph) ? UbergraphFunctionName : FunctionName;
 	TSharedPtr<FBlueprintFunctionContext>& Result = FunctionContexts.FindOrAdd(FunctionName);
 	if (!Result.IsValid())
 	{
 		Result = MakeShareable(new FBlueprintFunctionContext);
-		Result->InitialiseContextFromGraph(AsShared(), Graph);
+		Result->InitialiseContextFromGraph(AsShared(), FunctionName, Graph);
 	}
 	check(Result.IsValid());
 	return Result;
 }
 
+bool FBlueprintExecutionContext::HasProfilerDataForPin(const UEdGraphPin* GraphPin) const
+{
+	SCOPE_CYCLE_COUNTER(STAT_NodeLookupCost);
+	const UEdGraph* OuterGraph = GraphPin ? GraphPin->GetOwningNode()->GetTypedOuter<UEdGraph>() : nullptr;
+	if (OuterGraph)
+	{
+		TSharedPtr<FBlueprintFunctionContext> FunctionContext = GetFunctionContext(OuterGraph->GetFName());
+		if (FunctionContext.IsValid())
+		{
+			const FName PinName(FunctionContext->GetUniquePinName(GraphPin));
+			return FunctionContext->HasProfilerDataForNode(PinName);
+		}
+	}
+	return false;
+}
+
+TSharedPtr<FScriptExecutionNode> FBlueprintExecutionContext::GetProfilerDataForPin(const UEdGraphPin* GraphPin)
+{
+	SCOPE_CYCLE_COUNTER(STAT_NodeLookupCost);
+	TSharedPtr<FScriptExecutionNode> Result;
+	const UEdGraph* OuterGraph = GraphPin ? GraphPin->GetOwningNode()->GetTypedOuter<UEdGraph>() : nullptr;
+	if (OuterGraph)
+	{
+		TSharedPtr<FBlueprintFunctionContext> FunctionContext = GetFunctionContext(OuterGraph->GetFName());
+		if (FunctionContext.IsValid())
+		{
+			const FName PinName(FunctionContext->GetUniquePinName(GraphPin));
+			const FName UniquePinName(*FString::Printf(TEXT("%s_%s"), *GraphPin->GetOwningNode()->GetFName().ToString(), *PinName.ToString()));
+			if (FunctionContext->HasProfilerDataForNode(PinName))
+			{
+				Result = FunctionContext->GetProfilerDataForNode(PinName);
+			}
+			else
+			{
+				const UEdGraphNode* OwningNode = GraphPin->GetOwningNode();
+				if (FunctionContext->HasProfilerDataForNode(OwningNode->GetFName()))
+				{
+					Result = FunctionContext->GetProfilerDataForNode(OwningNode->GetFName());
+				}
+			}
+		}
+	}
+	return Result;
+}
+
 bool FBlueprintExecutionContext::HasProfilerDataForNode(const UEdGraphNode* GraphNode) const
 {
 	SCOPE_CYCLE_COUNTER(STAT_NodeLookupCost);
@@ -207,24 +327,36 @@ bool FBlueprintExecutionContext::MapBlueprintExecution()
 
 	if (BlueprintToMap && BPGC)
 	{
-		// Map all blueprint graphs and event entry points
-		TArray<UEdGraph*> Graphs;
-		TArray<UEdGraph*> FunctionGraphs;
-		TSet<UEdGraph*> Macros;
-		BlueprintToMap->GetAllGraphs(Graphs);
-		for (auto Graph : Graphs)
+		// Grab all ancestral blueprints used to generate this class.
+		TArray<UBlueprint*> InheritedBlueprints;
+		BlueprintToMap->GetBlueprintHierarchyFromClass(BPGC, InheritedBlueprints);
+		// Locate all blueprint graphs and event entry points
+		TMap<FName, UEdGraph*> FunctionGraphs;
+		for (auto CurrBlueprint : InheritedBlueprints)
 		{
-			FName FunctionName = Graph->GetFName();
-			FunctionName = FunctionName == UEdGraphSchema_K2::GN_EventGraph ? UbergraphFunctionName : FunctionName;
-			if (UFunction* ScriptFunction = BPGC->FindFunctionByName(FunctionName))
+			if (UBlueprintGeneratedClass* CurrBPGC = Cast<UBlueprintGeneratedClass>(CurrBlueprint->GeneratedClass))
 			{
-				FunctionGraphs.Add(Graph);
+				const FName CurrUbergraphName = Cast<UBlueprintGeneratedClass>(CurrBlueprint->GeneratedClass)->UberGraphFunction->GetFName();
+				TArray<UEdGraph*> Graphs;
+				CurrBlueprint->GetAllGraphs(Graphs);
+				for (auto Graph : Graphs)
+				{
+					FName FunctionName = Graph->GetFName();
+					FunctionName = FunctionName == UEdGraphSchema_K2::GN_EventGraph ? CurrUbergraphName : FunctionName;
+					if (!FunctionGraphs.Contains(FunctionName))
+					{
+						if (UFunction* ScriptFunction = CurrBPGC->FindFunctionByName(FunctionName))
+						{
+							FunctionGraphs.Add(FunctionName) = Graph;
+						}
+					}
+				}
 			}
 		}
 		// Create function stubs
-		for (auto Graph : FunctionGraphs)
+		for (auto GraphEntry : FunctionGraphs)
 		{
-			CreateFunctionContext(Graph);
+			CreateFunctionContext(GraphEntry.Key, GraphEntry.Value);
 		}
 		// Map the function context nodes
 		for (auto Context : FunctionContexts)
@@ -239,44 +371,61 @@ bool FBlueprintExecutionContext::MapBlueprintExecution()
 //////////////////////////////////////////////////////////////////////
 // FBlueprintFunctionContext
 
-void FBlueprintFunctionContext::InitialiseContextFromGraph(TSharedPtr<FBlueprintExecutionContext> BlueprintContextIn, UEdGraph* Graph)
+void FBlueprintFunctionContext::InitialiseContextFromGraph(TSharedPtr<FBlueprintExecutionContext> BlueprintContextIn, const FName FunctionNameIn, UEdGraph* Graph)
 {
 	if (Graph)
 	{
 		// Initialise Blueprint references
 		BlueprintContext = BlueprintContextIn;
-		UBlueprint* Blueprint = BlueprintContextIn->GetBlueprint().Get();
-		UBlueprintGeneratedClass* BPClass = BlueprintContextIn->GetBlueprintClass().Get();
+		UBlueprint* Blueprint = Graph->GetTypedOuter<UBlueprint>();
+		UBlueprintGeneratedClass* BPClass = Blueprint ? Cast<UBlueprintGeneratedClass>(Blueprint->GeneratedClass) : nullptr;
+		check (Blueprint && BPClass);
 		if (Blueprint && BPClass)
 		{
 			// Instantiate context
+			OwningBlueprint = Blueprint;
 			BlueprintClass = BPClass;
+			const bool bInheritedClass = BlueprintContextIn->GetBlueprintClass() != BPClass;
 			const bool bEventGraph = Graph->GetFName() == UEdGraphSchema_K2::GN_EventGraph;
-			FunctionName = bEventGraph && BPClass->UberGraphFunction ? BPClass->UberGraphFunction->GetFName() : Graph->GetFName();
+			GraphName = Graph->GetFName();
+			FunctionName = FunctionNameIn;
 			Function = BPClass->FindFunctionByName(FunctionName);
 			Function = Function.IsValid() ? Function : BPClass->UberGraphFunction;
 			if (bEventGraph)
 			{
 				// Map events
 				TArray<UK2Node_Event*> GraphEventNodes;
-				Graph->GetNodesOfClass<UK2Node_Event>(GraphEventNodes);
-
+				for (auto EventGraph : Blueprint->UbergraphPages)
+				{
+					EventGraph->GetNodesOfClass<UK2Node_Event>(GraphEventNodes);
+				}
 				for (auto EventNode : GraphEventNodes)
 				{
 					FName EventName = EventNode->GetFunctionName();
-					FScriptExecNodeParams EventParams;
-					EventParams.NodeName = EventName;
-					EventParams.ObservedObject = EventNode;
-					EventParams.DisplayName = EventNode->GetNodeTitle(ENodeTitleType::ListView);
-					EventParams.Tooltip = LOCTEXT("NavigateToEventLocationHyperlink_ToolTip", "Navigate to the Event");
-					EventParams.NodeFlags = EScriptExecutionNodeFlags::Event;
-					EventParams.IconColor = FLinearColor(0.91f, 0.16f, 0.16f);
-					const FSlateBrush* EventIcon = EventNode->ShowPaletteIconOnNode() ?	EventNode->GetIconAndTint(EventParams.IconColor).GetOptionalIcon() :
-																						FEditorStyle::GetBrush(TEXT("BlueprintProfiler.BPNode"));
-					EventParams.Icon = const_cast<FSlateBrush*>(EventIcon);
-					TSharedPtr<FScriptExecutionNode> EventExecNode = CreateExecutionNode(EventParams);
-					AddEntryPoint(EventExecNode);
-					BlueprintContextIn->AddEventNode(EventExecNode);
+					if (!BlueprintContext.Pin()->IsEventMapped(EventName))
+					{
+						FScriptExecNodeParams EventParams;
+						EventParams.NodeName = EventName;
+						EventParams.ObservedObject = EventNode;
+						EventParams.DisplayName = EventNode->GetNodeTitle(ENodeTitleType::ListView);
+						if (bInheritedClass)
+						{
+							EventParams.Tooltip = LOCTEXT("NavigateToInheritedEventLocationHyperlink_ToolTip", "Navigate to the Inherited Event");
+							EventParams.IconColor = FLinearColor(0.91f, 0.16f, 0.16f, 0.4f);
+						}
+						else
+						{
+							EventParams.Tooltip = LOCTEXT("NavigateToEventLocationHyperlink_ToolTip", "Navigate to the Event");
+							EventParams.IconColor = FLinearColor(0.91f, 0.16f, 0.16f);
+						}
+						EventParams.NodeFlags = EScriptExecutionNodeFlags::Event;
+						const FSlateBrush* EventIcon = EventNode->ShowPaletteIconOnNode() ?	EventNode->GetIconAndTint(EventParams.IconColor).GetOptionalIcon() :
+																							FEditorStyle::GetBrush(TEXT("BlueprintProfiler.BPNode"));
+						EventParams.Icon = const_cast<FSlateBrush*>(EventIcon);
+						TSharedPtr<FScriptExecutionNode> EventExecNode = CreateExecutionNode(EventParams);
+						AddEntryPoint(EventExecNode);
+						BlueprintContextIn->AddEventNode(AsShared(), EventExecNode);
+					}
 				}
 				// Create Input events
 				TArray<UK2Node*> InputEventNodes;
@@ -287,6 +436,13 @@ void FBlueprintFunctionContext::InitialiseContextFromGraph(TSharedPtr<FBlueprint
 				{
 					CreateInputEvents(BlueprintContextIn, InputEventNodes);
 				}
+				// Create async task events
+				TArray<UK2Node_BaseAsyncTask*> AsyncTaskNodes;
+				Graph->GetNodesOfClass<UK2Node_BaseAsyncTask>(AsyncTaskNodes);
+				if (AsyncTaskNodes.Num())
+				{
+					CreateAsyncTaskEvents(BlueprintContextIn, AsyncTaskNodes);
+				}
 			}
 			else
 			{
@@ -301,9 +457,17 @@ void FBlueprintFunctionContext::InitialiseContextFromGraph(TSharedPtr<FBlueprint
 					FunctionNodeParams.NodeName = FunctionName;
 					FunctionNodeParams.ObservedObject = FunctionEntry;
 					FunctionNodeParams.DisplayName = FText::FromName(FunctionName);
-					FunctionNodeParams.Tooltip = LOCTEXT("NavigateToFunctionLocationHyperlink_ToolTip", "Navigate to the Function");
+					if (bInheritedClass)
+					{
+						FunctionNodeParams.Tooltip = LOCTEXT("NavigateToInheritedFunctionLocationHyperlink_ToolTip", "Navigate to the Inherited Function");
+						FunctionNodeParams.IconColor = FLinearColor(0.91f, 0.16f, 0.16f, 0.4f);
+					}
+					else
+					{
+						FunctionNodeParams.Tooltip = LOCTEXT("NavigateToFunctionLocationHyperlink_ToolTip", "Navigate to the Function");
+						FunctionNodeParams.IconColor = FLinearColor(0.91f, 0.16f, 0.16f);
+					}
 					FunctionNodeParams.NodeFlags = EScriptExecutionNodeFlags::Event;
-					FunctionNodeParams.IconColor = FLinearColor(0.91f, 0.16f, 0.16f);
 					const FSlateBrush* Icon = FunctionEntry->ShowPaletteIconOnNode() ? FunctionEntry->GetIconAndTint(FunctionNodeParams.IconColor).GetOptionalIcon() :
 																						FEditorStyle::GetBrush(TEXT("BlueprintProfiler.BPNode"));
 					FunctionNodeParams.Icon = const_cast<FSlateBrush*>(Icon);
@@ -311,24 +475,36 @@ void FBlueprintFunctionContext::InitialiseContextFromGraph(TSharedPtr<FBlueprint
 					AddEntryPoint(FunctionEntryNode);
 					if (bConstructionScript)
 					{
-						BlueprintContextIn->AddEventNode(FunctionEntryNode);
+						BlueprintContextIn->AddEventNode(AsShared(), FunctionEntryNode);
 					}
 				}
 			}
-			// Map all instanced sub graphs
-			TArray<UK2Node_Tunnel*> GraphTunnels;
-			Graph->GetNodesOfClass<UK2Node_Tunnel>(GraphTunnels);
-			// Map sub graphs / composites and macros
-			for (auto Tunnel : GraphTunnels)
-			{
-				MapTunnelInstance(Tunnel);
-			}
 		}
 	}
 }
 
 void FBlueprintFunctionContext::MapFunction()
 {
+	// Map any macro instances
+	if (UBlueprint* BP = OwningBlueprint.Get())
+	{
+		TArray<UEdGraph*> Graphs;
+		BP->GetAllGraphs(Graphs);
+		for (auto Graph : Graphs)
+		{
+			if (Graph->GetFName() == GraphName)
+			{
+				TArray<UK2Node_Tunnel*> GraphTunnels;
+				Graph->GetNodesOfClass<UK2Node_Tunnel>(GraphTunnels);
+				// Map sub graphs / composites and macros
+				for (auto Tunnel : GraphTunnels)
+				{
+					MapTunnelInstance(Tunnel);
+				}
+				break;
+			}
+		}
+	}
 	// Map all entry point execution paths
 	for (auto EntryPoint : EntryPoints)
 	{
@@ -348,51 +524,55 @@ void FBlueprintFunctionContext::MapFunction()
 
 bool FBlueprintFunctionContext::DetectCyclicLinks(TSharedPtr<FScriptExecutionNode> ExecNode, TSet<TSharedPtr<FScriptExecutionNode>>& Filter)
 {
+	if (ExecNode->HasFlags(EScriptExecutionNodeFlags::PureStats))
+	{
+		return false;
+	}
 	if (Filter.Contains(ExecNode))
 	{
 		return true;
 	}
 	else
 	{
-		Filter.Add(ExecNode);
+		if (ExecNode->IsCyclicNode())
+		{
+			Filter.Add(ExecNode);
+		}
 		for (auto Child : ExecNode->GetChildNodes())
 		{
 			DetectCyclicLinks(Child, Filter);
 		}
 		for (auto LinkedNode : ExecNode->GetLinkedNodes())
 		{
-			if (!LinkedNode.Value->HasFlags(EScriptExecutionNodeFlags::PureStats))
+			TSharedPtr<FScriptExecutionNode> LinkedExecNode = LinkedNode.Value;
+			if (LinkedExecNode->HasFlags(EScriptExecutionNodeFlags::FunctionTunnel))
 			{
-				TSharedPtr<FScriptExecutionNode> LinkedExecNode = LinkedNode.Value;
-				if (LinkedExecNode->HasFlags(EScriptExecutionNodeFlags::FunctionTunnel))
+				TSharedPtr<FScriptExecutionNode> TunnelBoundry = LinkedExecNode->GetLinkedNodeByScriptOffset(LinkedNode.Key);
+				if (TunnelBoundry.IsValid())
 				{
-					TSharedPtr<FScriptExecutionNode> TunnelBoundry = LinkedExecNode->GetLinkedNodeByScriptOffset(LinkedNode.Key);
-					if (TunnelBoundry.IsValid())
+					TSharedPtr<FScriptExecutionNode> TunnelExec = TunnelBoundry->GetLinkedNodeByScriptOffset(LinkedNode.Key);
+					if (TunnelExec.IsValid())
 					{
-						TSharedPtr<FScriptExecutionNode> TunnelExec = TunnelBoundry->GetLinkedNodeByScriptOffset(LinkedNode.Key);
-						if (TunnelExec.IsValid())
-						{
-							ExecNode = TunnelExec;
-						}
+						ExecNode = TunnelExec;
 					}
 				}
-				TSet<TSharedPtr<FScriptExecutionNode>> BranchFilter(Filter);
-				if (DetectCyclicLinks(LinkedExecNode, BranchFilter))
-				{
-					// Break links and flag cycle linkage.
-					FScriptExecNodeParams CycleLinkParams;
-					const FString NodeName = FString::Printf(TEXT("CyclicLinkTo_%i_%s"), ExecutionNodes.Num(), *LinkedExecNode->GetName().ToString());
-					CycleLinkParams.NodeName = FName(*NodeName);
-					CycleLinkParams.ObservedObject = LinkedExecNode->GetObservedObject();
-					CycleLinkParams.DisplayName = LinkedExecNode->GetDisplayName();
-					CycleLinkParams.Tooltip = LOCTEXT("CyclicLink_ToolTip", "Cyclic Link");
-					CycleLinkParams.NodeFlags = EScriptExecutionNodeFlags::CyclicLinkage|EScriptExecutionNodeFlags::ExecPin|EScriptExecutionNodeFlags::InvalidTrace;
-					CycleLinkParams.IconColor = FLinearColor(0.91f, 0.16f, 0.16f);
-					const FSlateBrush* LinkIcon = LinkedExecNode->GetIcon();
-					CycleLinkParams.Icon = const_cast<FSlateBrush*>(LinkIcon);
-					TSharedPtr<FScriptExecutionNode> NewLink = CreateExecutionNode(CycleLinkParams);
-					ExecNode->AddLinkedNode(LinkedNode.Key, NewLink);
-				}
+			}
+			TSet<TSharedPtr<FScriptExecutionNode>> BranchFilter(Filter);
+			if (DetectCyclicLinks(LinkedExecNode, BranchFilter))
+			{
+				// Break links and flag cycle linkage.
+				FScriptExecNodeParams CycleLinkParams;
+				const FString NodeName = FString::Printf(TEXT("CyclicLinkTo_%i_%s"), ExecutionNodes.Num(), *LinkedExecNode->GetName().ToString());
+				CycleLinkParams.NodeName = FName(*NodeName);
+				CycleLinkParams.ObservedObject = LinkedExecNode->GetObservedObject();
+				CycleLinkParams.DisplayName = LinkedExecNode->GetDisplayName();
+				CycleLinkParams.Tooltip = LOCTEXT("CyclicLink_ToolTip", "Cyclic Link");
+				CycleLinkParams.NodeFlags = EScriptExecutionNodeFlags::CyclicLinkage|EScriptExecutionNodeFlags::ExecPin|EScriptExecutionNodeFlags::InvalidTrace;
+				CycleLinkParams.IconColor = FLinearColor(0.91f, 0.16f, 0.16f);
+				const FSlateBrush* LinkIcon = LinkedExecNode->GetIcon();
+				CycleLinkParams.Icon = const_cast<FSlateBrush*>(LinkIcon);
+				TSharedPtr<FScriptExecutionNode> NewLink = CreateExecutionNode(CycleLinkParams);
+				ExecNode->AddLinkedNode(LinkedNode.Key, NewLink);
 			}
 		}
 	}
@@ -401,20 +581,15 @@ bool FBlueprintFunctionContext::DetectCyclicLinks(TSharedPtr<FScriptExecutionNod
 
 void FBlueprintFunctionContext::CreateInputEvents(TSharedPtr<FBlueprintExecutionContext> BlueprintContextIn, const TArray<UK2Node*>& InputEventNodes)
 {
-	struct FInputEventDesc
-	{
-		UEdGraphNode* GraphNode;
-		UFunction* EventFunction;
-	};
-	TArray<FInputEventDesc> InputEventDescs;
+	TMap<const UEdGraphNode*, TArray<FName>> NodeEventDescs;
 	// Extract basic event info
 	if (UBlueprintGeneratedClass* BPClass = BlueprintClass.Get())
 	{
-		FInputEventDesc NewInputDesc;
 		for (auto EventNode : InputEventNodes)
 		{
 			if (EventNode)
 			{
+				TArray<FName>& Events = NodeEventDescs.FindOrAdd(EventNode);
 				if (UK2Node_InputAction* InputActionNode = Cast<UK2Node_InputAction>(EventNode))
 				{
 					const FString CustomEventName = FString::Printf(TEXT("InpActEvt_%s_%s"), *InputActionNode->InputActionName.ToString());
@@ -422,9 +597,7 @@ void FBlueprintFunctionContext::CreateInputEvents(TSharedPtr<FBlueprintExecution
 					{
 						if (FunctionIter->GetName().Contains(CustomEventName))
 						{
-							NewInputDesc.GraphNode = EventNode;
-							NewInputDesc.EventFunction = *FunctionIter;
-							InputEventDescs.Add(NewInputDesc);
+							Events.Add(FunctionIter->GetFName());
 						}
 					}
 				}
@@ -437,9 +610,7 @@ void FBlueprintFunctionContext::CreateInputEvents(TSharedPtr<FBlueprintExecution
 					{
 						if (FunctionIter->GetName().Contains(CustomEventName))
 						{
-							NewInputDesc.GraphNode = EventNode;
-							NewInputDesc.EventFunction = *FunctionIter;
-							InputEventDescs.Add(NewInputDesc);
+							Events.Add(FunctionIter->GetFName());
 						}
 					}
 				}
@@ -475,14 +646,12 @@ void FBlueprintFunctionContext::CreateInputEvents(TSharedPtr<FBlueprintExecution
 					}
 					for (auto Pin : ActivePins)
 					{
-						const FString CustomEventName = FString::Printf(TEXT("InpTchEvt_%s"), *Pin->GetName());
+						const FString CustomEventName = FString::Printf(TEXT("InpTchEvt_%s"), *GetPinName(Pin).ToString());
 						for (auto FunctionIter = TFieldIterator<UFunction>(BPClass); FunctionIter; ++FunctionIter)
 						{
 							if (FunctionIter->GetName().Contains(CustomEventName))
 							{
-								NewInputDesc.GraphNode = EventNode;
-								NewInputDesc.EventFunction = *FunctionIter;
-								InputEventDescs.Add(NewInputDesc);
+								Events.Add(FunctionIter->GetFName());
 							}
 						}
 					}
@@ -491,24 +660,106 @@ void FBlueprintFunctionContext::CreateInputEvents(TSharedPtr<FBlueprintExecution
 		}
 	}
 	// Build event contexts
-	for (auto InputEventDesc : InputEventDescs)
+	for (auto NodeEvents : NodeEventDescs)
 	{
-		UEdGraphNode* GraphNode = InputEventDesc.GraphNode;
-		FName EventName = InputEventDesc.EventFunction->GetFName();
+		// Setup the basic exec node params.
 		FScriptExecNodeParams EventParams;
-		EventParams.NodeName = EventName;
-		EventParams.ObservedObject = GraphNode;
+		EventParams.ObservedObject = NodeEvents.Key;
 		EventParams.OwningGraphName = NAME_None;
-		EventParams.DisplayName = GraphNode->GetNodeTitle(ENodeTitleType::ListView);
+		EventParams.DisplayName = NodeEvents.Key->GetNodeTitle(ENodeTitleType::ListView);
 		EventParams.Tooltip = LOCTEXT("NavigateToEventLocationHyperlink_ToolTip", "Navigate to the Event");
 		EventParams.NodeFlags = EScriptExecutionNodeFlags::Event;
 		EventParams.IconColor = FLinearColor(0.91f, 0.16f, 0.16f);
-		const FSlateBrush* EventIcon = GraphNode->ShowPaletteIconOnNode() ? GraphNode->GetIconAndTint(EventParams.IconColor).GetOptionalIcon() :
-																			FEditorStyle::GetBrush(TEXT("BlueprintProfiler.BPNode"));
+		const FSlateBrush* EventIcon = NodeEvents.Key->ShowPaletteIconOnNode() ? NodeEvents.Key->GetIconAndTint(EventParams.IconColor).GetOptionalIcon() :
+																				 FEditorStyle::GetBrush(TEXT("BlueprintProfiler.BPNode"));
 		EventParams.Icon = const_cast<FSlateBrush*>(EventIcon);
+		TSharedPtr<FScriptExecutionNode> EventExecNode;
+		// Create events for this node
+		for (FName EventName : NodeEvents.Value)
+		{
+			if (!EventExecNode.IsValid())
+			{
+				// Create the event node.
+				EventParams.NodeName = EventName;
+				EventExecNode = CreateExecutionNode(EventParams);
+				// Add entry points.
+				AddEntryPoint(EventExecNode);
+				BlueprintContextIn->AddEventNode(AsShared(), EventExecNode);
+			}
+			else
+			{
+				// Register the exec node under the event name.
+				ExecutionNodes.Add(EventName) = EventExecNode;
+				// Register the function context as a handler for the event.
+				BlueprintContextIn->RegisterEventContext(EventName, AsShared());
+			}
+		}
+	}
+}
+
+void FBlueprintFunctionContext::CreateAsyncTaskEvents(TSharedPtr<FBlueprintExecutionContext> BlueprintContextIn, const TArray<UK2Node_BaseAsyncTask*>& AsyncTaskNodes)
+{
+	struct FAsyncDelegateDesc
+	{
+		UEdGraphNode* DelegateNode;
+		UEdGraphPin* DelegatePin;
+		UFunction* EventFunction;
+	};
+	TArray<FAsyncDelegateDesc> AsyncDelegates;
+	// Extract basic event info
+	if (UBlueprintGeneratedClass* BPClass = BlueprintClass.Get())
+	{
+		FAsyncDelegateDesc NewDelegateDesc;
+		for (auto AsyncTaskNode : AsyncTaskNodes)
+		{
+			if (AsyncTaskNode)
+			{
+				for (auto Pin : AsyncTaskNode->Pins)
+				{
+					if (Pin->Direction == EGPD_Output && Pin->PinType.PinCategory == UEdGraphSchema_K2::PC_Exec)
+					{
+						FName PinName(GetPinName(Pin));
+						const FString PinDelegateName = FString::Printf(TEXT("%s_"), *PinName.ToString());
+						if (PinName != NAME_None)
+						{
+							for (auto FunctionIter = TFieldIterator<UFunction>(BPClass, EFieldIteratorFlags::ExcludeSuper); FunctionIter; ++FunctionIter)
+							{
+								if (FunctionIter->HasAllFunctionFlags(FUNC_BlueprintEvent|FUNC_BlueprintCallable) && FunctionIter->GetName().Contains(PinDelegateName))
+								{
+									NewDelegateDesc.DelegateNode = AsyncTaskNode;
+									NewDelegateDesc.DelegatePin = Pin;
+									NewDelegateDesc.EventFunction = *FunctionIter;
+									AsyncDelegates.Add(NewDelegateDesc);
+									break;
+								}
+							}
+						}
+					}
+				}
+			}
+		}
+	}
+	// Build event contexts
+	for (auto AsyncDelegateDesc : AsyncDelegates)
+	{
+		UEdGraphPin* DelegatePin = AsyncDelegateDesc.DelegatePin;
+		FName EventName = AsyncDelegateDesc.EventFunction->GetFName();
+		FScriptExecNodeParams EventParams;
+		EventParams.NodeName = EventName;
+		EventParams.ObservedObject = DelegatePin->GetOwningNode();
+		EventParams.ObservedPin = DelegatePin;
+		EventParams.OwningGraphName = NAME_None;
+		EventParams.DisplayName = DelegatePin->GetDisplayName();
+		EventParams.Tooltip = LOCTEXT("NavigateToEventLocationHyperlink_ToolTip", "Navigate to the Event");
+		EventParams.NodeFlags = EScriptExecutionNodeFlags::Event|EScriptExecutionNodeFlags::AsyncTaskDelegate;
+		EventParams.IconColor = FLinearColor(1.f, 1.f, 1.f, 0.8f);
+		const FSlateBrush* Icon = DelegatePin->LinkedTo.Num() ?	FEditorStyle::GetBrush(TEXT("BlueprintProfiler.BPPinConnected")) : 
+																FEditorStyle::GetBrush(TEXT("BlueprintProfiler.BPPinDisconnected"));
+		EventParams.Icon = const_cast<FSlateBrush*>(Icon);
 		TSharedPtr<FScriptExecutionNode> EventExecNode = CreateExecutionNode(EventParams);
+		const FName PinName = GetUniquePinName(DelegatePin);
+		ExecutionNodes.Add(PinName) = EventExecNode;
 		AddEntryPoint(EventExecNode);
-		BlueprintContextIn->AddEventNode(EventExecNode);
 	}
 }
 
@@ -555,7 +806,11 @@ TSharedPtr<FScriptExecutionNode> FBlueprintFunctionContext::MapNodeExecution(UEd
 				MappedNode->SetPureNodeScriptCodeRange(GetPureNodeScriptCodeRange(NodeToMap));
 			}
 			// Determine a few node characteristics to guide processing.
-			if (ExecPins.Num() > 1)
+			if (const UK2Node_BaseAsyncTask* AsyncNode = Cast<UK2Node_BaseAsyncTask>(NodeToMap))
+			{
+				MappedNode->AddFlags(ExecPins.Num() ? (EScriptExecutionNodeFlags::AsyncTask|EScriptExecutionNodeFlags::SequentialBranch) : EScriptExecutionNodeFlags::AsyncTask);
+			}
+			else if (ExecPins.Num() > 1)
 			{
 				const bool bExecSequence = NodeToMap->IsA<UK2Node_ExecutionSequence>();
 				MappedNode->AddFlags(bExecSequence ? EScriptExecutionNodeFlags::SequentialBranch : EScriptExecutionNodeFlags::ConditionalBranch);
@@ -612,9 +867,11 @@ void FBlueprintFunctionContext::MapInputPins(TSharedPtr<FScriptExecutionNode> Ex
 		// If this input pin is linked to a pure node in the source graph, create and map all known execution paths for it.
 		for (auto LinkedPin : InputPin->LinkedTo)
 		{
+			// Pass through non-relevant (e.g. reroute) nodes.
+			UK2Node* OwningNode = FBlueprintEditorUtils::FindFirstCompilerRelevantNode(LinkedPin);
+
 			// Note: Intermediate pure nodes can have output pins that masquerade as impure node output pins when links are "moved" from the source graph (thus
 			// resulting in a false association here with one or more script code offsets), so we must first ensure that the link is really to a pure node output.
-			UK2Node* OwningNode = Cast<UK2Node>(LinkedPin->GetOwningNode());
 			if (OwningNode && OwningNode->IsNodePure())
 			{
 				GetAllCodeLocationsFromPin(LinkedPin, PinScriptCodeOffsets);
@@ -630,8 +887,8 @@ void FBlueprintFunctionContext::MapInputPins(TSharedPtr<FScriptExecutionNode> Ex
 						PureChainParams.DisplayName = LOCTEXT("PureChain_DisplayName", "Pure Time");
 						PureChainParams.Tooltip = LOCTEXT("PureChain_ToolTip", "Expand pure node timing");
 						PureChainParams.NodeFlags = EScriptExecutionNodeFlags::PureChain;
-						PureChainParams.IconColor = FLinearColor(1.f, 1.f, 1.f, 0.8f);
-						const FSlateBrush* Icon = FEditorStyle::GetBrush(TEXT("BlueprintProfiler.BPNode"));
+						PureChainParams.IconColor = FLinearColor(0.2f, 1.f, 0.2f);
+						const FSlateBrush* Icon = FEditorStyle::GetBrush(TEXT("BlueprintProfiler.PureNode"));
 						PureChainParams.Icon = const_cast<FSlateBrush*>(Icon);
 						PureChainRootNode = CreateExecutionNode(PureChainParams);
 						PureChainRootNode->SetPureNodeScriptCodeRange(ExecNode->GetPureNodeScriptCodeRange());
@@ -652,32 +909,47 @@ void FBlueprintFunctionContext::MapInputPins(TSharedPtr<FScriptExecutionNode> Ex
 
 void FBlueprintFunctionContext::MapExecPins(TSharedPtr<FScriptExecutionNode> ExecNode, const TArray<UEdGraphPin*>& Pins)
 {
+	const bool bAsyncTask = ExecNode->HasFlags(EScriptExecutionNodeFlags::AsyncTask);
 	const bool bBranchedExecution = ExecNode->IsBranch();
 	int32 NumUnwiredPins = 0;
 	const UEdGraphPin* FinalExecPin = nullptr;
 	for (auto Pin : Pins)
 	{
+		const FName PinName = GetUniquePinName(Pin);
 		int32 PinScriptCodeOffset = GetCodeLocationFromPin(Pin);
 		const bool bInvalidTrace = PinScriptCodeOffset == INDEX_NONE;
 		PinScriptCodeOffset = bInvalidTrace ? NumUnwiredPins++ : PinScriptCodeOffset;
 		TSharedPtr<FScriptExecutionNode> PinExecNode;
-		const UEdGraphPin* LinkedPin = Pin->LinkedTo.Num() > 0 ? Pin->LinkedTo[0] : nullptr;
+		// Note: Pass through non-relevant (e.g. reroute) nodes here as they're not compiled and thus do not need to be mapped for profiling.
+		const UEdGraphPin* LinkedPin = Pin->LinkedTo.Num() > 0 ? FBlueprintEditorUtils::FindFirstCompilerRelevantLinkedPin(Pin->LinkedTo[0]) : nullptr;
 		UEdGraphNode* LinkedPinNode = LinkedPin ? LinkedPin->GetOwningNode() : nullptr;
 		const bool bTunnelBoundry = LinkedPinNode ? LinkedPinNode->IsA<UK2Node_Tunnel>() : false;
 		// Create any neccesary dummy pins for branches
 		if (bBranchedExecution)
 		{
-			FScriptExecNodeParams LinkNodeParams;
-			LinkNodeParams.NodeName = Pin->GetFName();
-			LinkNodeParams.ObservedObject = Pin;
-			LinkNodeParams.DisplayName = Pin->GetDisplayName();
-			LinkNodeParams.Tooltip = LOCTEXT("ExecPin_ExpandExecutionPath_ToolTip", "Expand execution path");
-			LinkNodeParams.NodeFlags = !bInvalidTrace ? EScriptExecutionNodeFlags::ExecPin : (EScriptExecutionNodeFlags::ExecPin|EScriptExecutionNodeFlags::InvalidTrace);
-			LinkNodeParams.IconColor = FLinearColor(1.f, 1.f, 1.f, 0.8f);
-			const FSlateBrush* Icon = LinkedPin ?	FEditorStyle::GetBrush(TEXT("BlueprintProfiler.BPPinConnected")) : 
-													FEditorStyle::GetBrush(TEXT("BlueprintProfiler.BPPinDisconnected"));
-			LinkNodeParams.Icon = const_cast<FSlateBrush*>(Icon);
-			PinExecNode = CreateExecutionNode(LinkNodeParams);
+			if (bAsyncTask)
+			{
+				if (TSharedPtr<FScriptExecutionNode>* DelegateNode = ExecutionNodes.Find(PinName))
+				{
+					PinExecNode = *DelegateNode;
+					PinExecNode->AddFlags(EScriptExecutionNodeFlags::ExecPin);
+				}
+			}
+			if (!PinExecNode.IsValid())
+			{
+				FScriptExecNodeParams LinkNodeParams;
+				LinkNodeParams.NodeName = PinName;
+				LinkNodeParams.ObservedObject = Pin->GetOwningNode();
+				LinkNodeParams.DisplayName = Pin->GetDisplayName();
+				LinkNodeParams.ObservedPin = Pin;
+				LinkNodeParams.Tooltip = LOCTEXT("ExecPin_ExpandExecutionPath_ToolTip", "Expand execution path");
+				LinkNodeParams.NodeFlags = !bInvalidTrace ? EScriptExecutionNodeFlags::ExecPin : (EScriptExecutionNodeFlags::ExecPin|EScriptExecutionNodeFlags::InvalidTrace);
+				LinkNodeParams.IconColor = FLinearColor(1.f, 1.f, 1.f, 0.8f);
+				const FSlateBrush* Icon = LinkedPin ?	FEditorStyle::GetBrush(TEXT("BlueprintProfiler.BPPinConnected")) : 
+														FEditorStyle::GetBrush(TEXT("BlueprintProfiler.BPPinDisconnected"));
+				LinkNodeParams.Icon = const_cast<FSlateBrush*>(Icon);
+				PinExecNode = CreateExecutionNode(LinkNodeParams);
+			}
 			ExecNode->AddLinkedNode(PinScriptCodeOffset, PinExecNode);
 			if (!bTunnelBoundry)
 			{
@@ -690,7 +962,17 @@ void FBlueprintFunctionContext::MapExecPins(TSharedPtr<FScriptExecutionNode> Exe
 		}
 		else
 		{
-			PinExecNode = ExecNode;
+			if (bAsyncTask)
+			{
+				if (TSharedPtr<FScriptExecutionNode>* DelegateNode = ExecutionNodes.Find(PinName))
+				{
+					PinExecNode = *DelegateNode;
+				}
+			}
+			if (!PinExecNode.IsValid())
+			{
+				PinExecNode = ExecNode;
+			}
 			if (!bTunnelBoundry)
 			{
 				TSharedPtr<FScriptExecutionNode> LinkedPinExecNode = MapNodeExecution(LinkedPinNode);
@@ -705,7 +987,7 @@ void FBlueprintFunctionContext::MapExecPins(TSharedPtr<FScriptExecutionNode> Exe
 		{
 			if (StagingTunnelName != NAME_None)
 			{
-				const FName TunnelExitName(*LinkedPin->PinName);
+				const FName TunnelExitName(GetPinName(LinkedPin));
 				TSharedPtr<FScriptExecutionNode> TunnelExitNode = GetProfilerDataForNode(TunnelExitName);
 				// Check for re-entrant exit site
 				if (ExecNode->HasFlags(EScriptExecutionNodeFlags::SequentialBranch))
@@ -740,7 +1022,7 @@ void FBlueprintFunctionContext::MapExecPins(TSharedPtr<FScriptExecutionNode> Exe
 			}
 			else
 			{
-				const FName TunnelName(*LinkedPin->PinName);
+				const FName TunnelName(GetPinName(LinkedPin));
 				TSharedPtr<FScriptExecutionNode> TunnelEntryNode = MapTunnelEntry(TunnelName, Pin, LinkedPin, PinScriptCodeOffset);
 				if (TunnelEntryNode.IsValid())
 				{
@@ -758,6 +1040,41 @@ void FBlueprintFunctionContext::MapExecPins(TSharedPtr<FScriptExecutionNode> Exe
 	}
 }
 
+FName FBlueprintFunctionContext::GetPinName(const UEdGraphPin* Pin) const
+{
+	FName PinName(NAME_None);
+	if (Pin)
+	{
+		FString PinString = Pin->PinName;
+		if (PinString.IsEmpty())
+		{
+			int32 PinTypeIndex = INDEX_NONE;
+			for (auto NodePin : Pin->GetOwningNode()->Pins)
+			{
+				if (NodePin->Direction == Pin->Direction && NodePin->PinType.PinCategory == Pin->PinType.PinCategory)
+				{
+					PinTypeIndex++;
+				}
+			}
+			PinString = PinTypeIndex > 0 ? FString::Printf(TEXT("%s%i"), *Pin->PinType.PinCategory, PinTypeIndex) : Pin->PinType.PinCategory;
+		}
+		PinName = FName(*PinString);
+	}
+	check (PinName != NAME_None);
+	return PinName;
+}
+
+FName FBlueprintFunctionContext::GetUniquePinName(const UEdGraphPin* Pin) const
+{
+	FName PinName = GetPinName(Pin);
+	if (Pin)
+	{
+		PinName = FName(*FString::Printf(TEXT("%s_%s"), *Pin->GetOwningNode()->GetFName().ToString(), *PinName.ToString()));
+	}
+	check (PinName != NAME_None);
+	return PinName;
+}
+
 TSharedPtr<FScriptExecutionNode> FBlueprintFunctionContext::MapTunnelEntry(const FName TunnelName, const UEdGraphPin* ExecPin, const UEdGraphPin* TunnelPin, const int32 PinScriptOffset)
 {
 	// Retrieve tunnel function context
@@ -765,7 +1082,7 @@ TSharedPtr<FScriptExecutionNode> FBlueprintFunctionContext::MapTunnelEntry(const
 	TSharedPtr<FBlueprintFunctionContext> TunnelContext = GetTunnelContextFromNode(TunnelNode);
 	const int32 TunnelId = TunnelContext->GetTunnelIdFromName(TunnelName);
 	check(TunnelId != INDEX_NONE);
-	// Create node for collapsed tunnel/macro
+	// Create a tunnel instance node for the tunnel/macro to centralise stats within and aid stat lookup - for external lookup purposes by node name.
 	TSharedPtr<FScriptExecutionTunnelInstance> TunnelInstance = GetTypedProfilerDataForNode<FScriptExecutionTunnelInstance>(TunnelNode->GetFName());
 	if (!TunnelInstance.IsValid())
 	{
@@ -781,7 +1098,8 @@ TSharedPtr<FScriptExecutionNode> FBlueprintFunctionContext::MapTunnelEntry(const
 		TunnelInstanceParams.Icon = const_cast<FSlateBrush*>(TunnelIcon);
 		TunnelInstance = CreateTypedExecutionNode<FScriptExecutionTunnelInstance>(TunnelInstanceParams);
 	}
-	// Create tunnel customised entry node
+	// Create a customised tunnel entry nodes,
+	// This maps each unique entry point per tunnel instance so we can only include the entry and exit sites in use.
 	FScriptExecNodeParams TunnelParams;
 	TunnelParams.NodeName = FName(*FString::Printf(TEXT("%s_EntryPointId%i"), *TunnelNode->GetFName().ToString(), TunnelId));
 	TunnelParams.ObservedObject = TunnelNode;
@@ -793,30 +1111,35 @@ TSharedPtr<FScriptExecutionNode> FBlueprintFunctionContext::MapTunnelEntry(const
 																		  FEditorStyle::GetBrush(TEXT("BlueprintProfiler.BPNode"));
 	TunnelParams.Icon = const_cast<FSlateBrush*>(TunnelIcon);
 	FScriptExecutionTunnelEntry::ETunnelType TunnelType = TunnelContext->GetTunnelType(TunnelName);
-	TSharedPtr<FScriptExecutionTunnelEntry> TunnelCallNode = MakeShareable(new FScriptExecutionTunnelEntry(TunnelParams, TunnelNode, TunnelType));
-	ExecutionNodes.FindOrAdd(TunnelParams.NodeName) = TunnelCallNode;
-	TunnelInstance->AddCustomEntryPoint(TunnelCallNode);
+	TSharedPtr<FScriptExecutionTunnelEntry> CustomTunnelEntryNode = MakeShareable(new FScriptExecutionTunnelEntry(TunnelParams, TunnelNode, TunnelType));
+	ExecutionNodes.FindOrAdd(TunnelParams.NodeName) = CustomTunnelEntryNode;
+	// Add the custom entry point to the tunnel instance ( for stat updates )
+	TunnelInstance->AddCustomEntryPoint(CustomTunnelEntryNode);
+	// Add the active entry point ( pin node ) to the custom tunnel entry node
 	TSharedPtr<FScriptExecutionNode> EntryPinNode = TunnelContext->GetProfilerDataForNode(TunnelName);
 	check (EntryPinNode.IsValid());
-	TunnelCallNode->AddChildNode(EntryPinNode);
+	CustomTunnelEntryNode->AddChildNode(EntryPinNode);
 	if (!EntryPinNode->GetNumChildren())
 	{
+		// Map the linked sites to the entry point on first encounter. ( this should probably be done during macro mapping )
 		TSharedPtr<FScriptExecutionNode> FirstTunnelNode = TunnelContext->GetTunnelEntryPoint(TunnelName);
 		check(FirstTunnelNode.IsValid());
 		EntryPinNode->AddChildNode(FirstTunnelNode);
 	}
-	// Map exit points
+	// Map tunnel exit points
 	TArray<TSharedPtr<FScriptExecutionNode>> TunnelExitPoints;
 	TunnelContext->GetTunnelExitPoints(TunnelName, TunnelExitPoints);
+	// Build a pin name to exit pin map ( located on the tunnel node rather than in the tunnel graph )
 	TMap<FName, const UEdGraphPin*> NameToPinMap;
 	for (auto ExitPin : TunnelNode->Pins)
 	{
 		if (ExitPin->PinType.PinCategory == UEdGraphSchema_K2::PC_Exec && ExitPin->LinkedTo.Num())
 		{
-			NameToPinMap.Add(FName(*ExitPin->PinName)) = ExitPin;
+			NameToPinMap.Add(GetPinName(ExitPin)) = ExitPin;
 		}
 	}
 	// Register tunnel input pin node using the script offset from the first script offset emitted pin.
+	// This is to try and improve tunnel entry detection for stats.
 	if (const UEdGraphPin** InputPin = NameToPinMap.Find(TunnelName))
 	{
 		if ((*InputPin)->LinkedTo.Num())
@@ -825,55 +1148,61 @@ TSharedPtr<FScriptExecutionNode> FBlueprintFunctionContext::MapTunnelEntry(const
 			{
 				const int32 EntryScriptCodeOffset = GetCodeLocationFromPin(InputLinkedPin);
 				ScriptOffsetToPins.Add(EntryScriptCodeOffset) = *InputPin;
-				TunnelEntrySites.Add(EntryScriptCodeOffset) = TunnelCallNode;
+				TunnelEntrySites.Add(EntryScriptCodeOffset) = CustomTunnelEntryNode;
 			}
 		}
 	}
+	// Map the tunnel exit sites to the exit sites for this instance and determine script offsets.
 	for (auto CurrExitPoint : TunnelExitPoints)
 	{
 		if (const UEdGraphPin** ExitPin = NameToPinMap.Find(CurrExitPoint->GetName()))
 		{
 			for (auto LinkedExitPin : (*ExitPin)->LinkedTo)
 			{
+				// This is the unique script offset for this instance's exit site.
 				const int32 ExitScriptCodeOffset = GetCodeLocationFromPin(*ExitPin);
-				// Add entry site for exit script offset
-				TunnelExitSites.Add(ExitScriptCodeOffset) = TunnelCallNode;
-				// map tunnel exit links
+				check (ExitScriptCodeOffset != INDEX_NONE);
+				// Add entry site for exit script offset - not sure what this is used for now - if anything.
+				TunnelExitSites.Add(ExitScriptCodeOffset) = CustomTunnelEntryNode;
+				// Pass through any non-relevant (e.g. reroute) nodes.
+				LinkedExitPin = FBlueprintEditorUtils::FindFirstCompilerRelevantLinkedPin(LinkedExitPin);
+				// Continue mapping exit sites
 				UEdGraphNode* LinkedExitNode = LinkedExitPin->GetOwningNode();
 				TSharedPtr<FScriptExecutionNode> LinkedExecNode;
 				if (LinkedExitNode->IsA<UK2Node_Tunnel>())
 				{
-					const FName LinkedTunnelName(*LinkedExitPin->PinName);
+					// Map being wired into another tunnel.
+					const FName LinkedTunnelName(GetPinName(LinkedExitPin));
 					LinkedExecNode = MapTunnelEntry(LinkedTunnelName, *ExitPin, LinkedExitPin, ExitScriptCodeOffset);
 				}
 				else
 				{
+					// Map forward normally.
 					LinkedExecNode = MapNodeExecution(LinkedExitPin->GetOwningNode());
 				}
 				if (LinkedExecNode.IsValid())
 				{
-					TunnelCallNode->AddExitSite(CurrExitPoint, ExitScriptCodeOffset);
 					if (CurrExitPoint->HasFlags(EScriptExecutionNodeFlags::ReEntrantTunnelPin))
 					{
+						// Add the linked nodes into the CurExitPoint
 						CurrExitPoint->AddLinkedNode(ExitScriptCodeOffset, LinkedExecNode);
 					}
 					else
 					{
+						// Add the linked nodes into the CurExitPoint and as a linked node to the custom tunnel entry.
 						CurrExitPoint->AddFlags(EScriptExecutionNodeFlags::TunnelFinalExitPin);
 						LinkedExecNode->AddFlags(EScriptExecutionNodeFlags::TunnelFinalExitPin);
-						TunnelCallNode->AddLinkedNode(ExitScriptCodeOffset, LinkedExecNode);
+						CustomTunnelEntryNode->AddLinkedNode(ExitScriptCodeOffset, LinkedExecNode);
 						CurrExitPoint->AddLinkedNode(ExitScriptCodeOffset, LinkedExecNode);
 					}
+					// Notify the custom tunnel entry node that this is a valid exit site for this instance.
+					CustomTunnelEntryNode->AddExitSite(CurrExitPoint, ExitScriptCodeOffset);
 				}
 			}
 		}
 	}
-	//if (TunnelType == FScriptExecutionTunnelEntry::MultiIOTunnel)
-	//{
-	//	TunnelCallNode->BuildExitBranches();
-	//}
 	// Return the tunnel entry.
-	return TunnelCallNode;
+	return CustomTunnelEntryNode;
 }
 
 bool FBlueprintFunctionContext::HasProfilerDataForNode(const FName NodeName) const
@@ -920,15 +1249,16 @@ const UEdGraphNode* FBlueprintFunctionContext::GetNodeFromCodeLocation(const int
 
 const UEdGraphPin* FBlueprintFunctionContext::GetPinFromCodeLocation(const int32 ScriptOffset)
 {
-	TWeakObjectPtr<const UEdGraphPin>& Result = ScriptOffsetToPins.FindOrAdd(ScriptOffset);
-	if (!Result.IsValid() && BlueprintClass.IsValid())
+	FEdGraphPinReference& Result = ScriptOffsetToPins.FindOrAdd(ScriptOffset);
+	const UEdGraphPin* Return = Result.Get();
+	if (!Return && BlueprintClass.IsValid())
 	{
 		if (const UEdGraphPin* GraphPin = BlueprintClass->GetDebugData().FindSourcePinFromCodeLocation(Function.Get(), ScriptOffset))
 		{
-			Result = GraphPin;
+			Return = GraphPin;
 		}
 	}
-	return Result.Get();
+	return Return;
 }
 
 const int32 FBlueprintFunctionContext::GetCodeLocationFromPin(const UEdGraphPin* Pin) const
@@ -1012,8 +1342,29 @@ void FBlueprintFunctionContext::MapTunnelInstance(UK2Node_Tunnel* TunnelInstance
 		TunnelGraphContext = BlueprintContext.Pin()->GetFunctionContext(TunnelGraph->GetFName());
 		if (!TunnelGraphContext.IsValid())
 		{
+			// Drill down for any nested macros
+			TArray<UK2Node_Tunnel*> GraphTunnels;
+			TunnelGraph->GetNodesOfClass<UK2Node_Tunnel>(GraphTunnels);
+			// Map sub graphs / composites and macros
+			for (auto ChildTunnel : GraphTunnels)
+			{
+				UEdGraph* ChildTunnelGraph = nullptr;
+				if (UK2Node_MacroInstance* MacroInstance = Cast<UK2Node_MacroInstance>(ChildTunnel))
+				{
+					ChildTunnelGraph = MacroInstance->GetMacroGraph();
+				}
+				else if (UK2Node_Composite* CompositeInstance = Cast<UK2Node_Composite>(ChildTunnel))
+				{
+					ChildTunnelGraph = CompositeInstance->BoundGraph;
+				}
+				if (ChildTunnelGraph && ChildTunnelGraph != TunnelGraph)
+				{
+					MapTunnelInstance(ChildTunnel);
+				}
+			}
 			// Create the sub graph context
 			TunnelGraphContext = MakeShareable(new FBlueprintFunctionContext);
+			TunnelGraphContext->GraphName = TunnelGraph->GetFName();
 			TunnelGraphContext->FunctionName = TunnelGraph->GetFName();
 			TunnelGraphContext->BlueprintContext = BlueprintContext;
 			TunnelGraphContext->Function = Function;
@@ -1025,16 +1376,19 @@ void FBlueprintFunctionContext::MapTunnelInstance(UK2Node_Tunnel* TunnelInstance
 			for (auto EntryPoint : TunnelGraphContext->EntryPoints)
 			{
 				TArray<TSharedPtr<FScriptExecutionNode>>& MappedExits = TunnelGraphContext->TunnelExitPointMap.FindOrAdd(EntryPoint->GetName());
-				if (const UEdGraphPin* EntryPointPin = Cast<UEdGraphPin>(EntryPoint->GetObservedObject()))
+				if (const UEdGraphPin* EntryPointPin = EntryPoint->GetObservedPin())
 				{
 					TunnelGraphContext->StagingTunnelName = EntryPoint->GetName();
 					for (auto LinkedPin : EntryPointPin->LinkedTo)
 					{
+						// Pass through any non-relevant (e.g. reroute) nodes.
+						LinkedPin = FBlueprintEditorUtils::FindFirstCompilerRelevantLinkedPin(LinkedPin);
+
 						const UEdGraphNode* LinkedNode = LinkedPin->GetOwningNode();
 						if (LinkedNode->IsA<UK2Node_Tunnel>())
 						{
 							// This handles input linked directly to output
-							TSharedPtr<FScriptExecutionNode> OutputPinNode = TunnelGraphContext->GetProfilerDataForNode(FName(*LinkedPin->PinName));
+							TSharedPtr<FScriptExecutionNode> OutputPinNode = TunnelGraphContext->GetProfilerDataForNode(GetPinName(LinkedPin));
 							check(OutputPinNode.IsValid());
 							MappedExits.Add(OutputPinNode);
 						}
@@ -1131,9 +1485,10 @@ void FBlueprintFunctionContext::MapTunnelIO(UEdGraph* TunnelGraph)
 				{
 					// Create entry point
 					FScriptExecNodeParams LinkParams;
-					LinkParams.NodeName = FName(*Pin->PinName);
-					LinkParams.ObservedObject = Pin;
-					LinkParams.DisplayName = Pin->GetDisplayName().IsEmpty() ? FText::FromString(*Pin->PinName) : Pin->GetDisplayName();
+					LinkParams.NodeName = GetPinName(Pin);
+					LinkParams.ObservedObject = Pin->GetOwningNode();
+					LinkParams.ObservedPin = Pin;
+					LinkParams.DisplayName = Pin->GetDisplayName().IsEmpty() ? FText::FromName(LinkParams.NodeName) : Pin->GetDisplayName();
 					LinkParams.Tooltip = LOCTEXT("ExecPin_ExpandTunnelEntryPoint_ToolTip", "Expand tunnel entry point");
 					LinkParams.NodeFlags = EScriptExecutionNodeFlags::TunnelEntryPin;
 					LinkParams.IconColor = FLinearColor(1.f, 1.f, 1.f, 0.8f);
@@ -1146,9 +1501,10 @@ void FBlueprintFunctionContext::MapTunnelIO(UEdGraph* TunnelGraph)
 				{
 					// Create exit point
 					FScriptExecNodeParams LinkParams;
-					LinkParams.NodeName = FName(*Pin->PinName);
-					LinkParams.ObservedObject = Pin;
-					LinkParams.DisplayName = Pin->GetDisplayName().IsEmpty() ? FText::FromString(*Pin->PinName) : Pin->GetDisplayName();
+					LinkParams.NodeName = GetPinName(Pin);
+					LinkParams.ObservedObject = Pin->GetOwningNode();
+					LinkParams.ObservedPin = Pin;
+					LinkParams.DisplayName = Pin->GetDisplayName().IsEmpty() ? FText::FromName(LinkParams.NodeName) : Pin->GetDisplayName();
 					LinkParams.Tooltip = LOCTEXT("ExecPin_ExpandTunnelExitPoint_ToolTip", "Expand tunnel exit point");
 					LinkParams.NodeFlags = EScriptExecutionNodeFlags::TunnelExitPin;
 					LinkParams.IconColor = FLinearColor(1.f, 1.f, 1.f, 0.8f);
@@ -1206,35 +1562,45 @@ void FBlueprintFunctionContext::GetTunnelsFromExitSite(const int32 ScriptCodeOff
 
 bool FScriptEventPlayback::Process(const TArray<FScriptInstrumentedEvent>& SignalData, const int32 StartIdx, const int32 StopIdx)
 {
-	struct PerNodeEventHelper
-	{
-		TArray<FTracePath> EntryTracePaths;
-		TArray<FTracePath> InputTracePaths;
-		TSharedPtr<FScriptExecutionNode> ExecNode;
-		TSharedPtr<FBlueprintFunctionContext> FunctionContext;
-		TArray<FScriptInstrumentedEvent> UnproccessedEvents;
-	};
-
 	const int32 NumEvents = (StopIdx+1) - StartIdx;
 	ProcessingState = EEventProcessingResult::Failed;
 	const bool bEventIsResuming = SignalData[StartIdx].IsResumeEvent() ;
+	IBlueprintProfilerInterface* Profiler = FModuleManager::GetModulePtr<IBlueprintProfilerInterface>("BlueprintProfiler");
 
 	if (BlueprintContext.IsValid() && InstanceName != NAME_None)
 	{
 		check(SignalData[StartIdx].IsEvent());
 		EventName = bEventIsResuming ? EventName : SignalData[StartIdx].GetFunctionName();
-		FName CurrentFunctionName = BlueprintContext->GetEventFunctionName(SignalData[StartIdx].GetFunctionName());
+		CurrentFunctionName = BlueprintContext->GetEventFunctionName(SignalData[StartIdx].GetFunctionName());
 		TSharedPtr<FBlueprintFunctionContext> FunctionContext = BlueprintContext->GetFunctionContext(CurrentFunctionName);
 		TSharedPtr<FScriptExecutionNode> EventNode = FunctionContext->GetProfilerDataForNode(EventName);
 		// Find Associated graph nodes and submit into a node map for later processing.
-		TMap<const UEdGraphNode*, PerNodeEventHelper> NodeMap;
+		TMap<const UEdGraphNode*, NodeSignalHelper> NodeMap;
 		ProcessingState = EEventProcessingResult::Success;
 		int32 LastEventIdx = StartIdx;
 		const int32 EventStartOffset = SignalData[StartIdx].IsResumeEvent() ? 3 : 1;
 
 		for (int32 SignalIdx = StartIdx + EventStartOffset; SignalIdx < StopIdx; ++SignalIdx)
 		{
+			// Handle midstream context switches.
 			const FScriptInstrumentedEvent& CurrSignal = SignalData[SignalIdx];
+			if (CurrSignal.GetType() == EScriptInstrumentation::Class)
+			{
+				// Update the current mapped blueprint context.
+				BlueprintContext = Profiler->GetBlueprintContext(CurrSignal.GetObjectPath());
+
+				// Skip to the next signal.
+				continue;
+			}
+			else if (CurrSignal.GetType() == EScriptInstrumentation::Instance)
+			{
+				// Update the current mapped instance name.
+				InstanceName = BlueprintContext->MapBlueprintInstance(CurrSignal.GetObjectPath());
+
+				// Skip to the next signal.
+				continue;
+			}
+
 			// Update script function.
 			if (CurrentFunctionName != CurrSignal.GetFunctionName())
 			{
@@ -1244,12 +1610,14 @@ bool FScriptEventPlayback::Process(const TArray<FScriptInstrumentedEvent>& Signa
 			}
 			if (const UEdGraphNode* GraphNode = FunctionContext->GetNodeFromCodeLocation(CurrSignal.GetScriptCodeOffset()))
 			{
-				PerNodeEventHelper& NodeEntry = NodeMap.FindOrAdd(GraphNode);
-				if (!NodeEntry.FunctionContext.IsValid())
+				NodeSignalHelper& CurrentNodeData = NodeMap.FindOrAdd(GraphNode);
+				// Initialise the current node context.
+				if (!CurrentNodeData.FunctionContext.IsValid())
 				{
-					NodeEntry.FunctionContext = FunctionContext;
-					NodeEntry.ExecNode = FunctionContext->GetProfilerDataForNode(GraphNode->GetFName());
-					check(NodeEntry.ExecNode.IsValid());
+					CurrentNodeData.BlueprintContext = BlueprintContext;
+					CurrentNodeData.FunctionContext = FunctionContext;
+					CurrentNodeData.ImpureNode = FunctionContext->GetProfilerDataForNode(GraphNode->GetFName());
+					check(CurrentNodeData.ImpureNode.IsValid());
 				}
 				switch (CurrSignal.GetType())
 				{
@@ -1260,26 +1628,17 @@ bool FScriptEventPlayback::Process(const TArray<FScriptInstrumentedEvent>& Signa
 							if (FunctionContext->HasProfilerDataForNode(Pin->GetOwningNode()->GetFName()))
 							{
 								TracePath.AddExitPin(CurrSignal.GetScriptCodeOffset());
-								NodeEntry.InputTracePaths.Insert(TracePath, 0);
+								CurrentNodeData.InputTracePaths.Insert(TracePath, 0);
 							}
 						}
-						NodeEntry.UnproccessedEvents.Add(CurrSignal);
-						// Add Trace History
-						FBlueprintExecutionTrace& NewTraceEvent = BlueprintContext->AddNewTraceHistoryEvent();
-						NewTraceEvent.TraceType = CurrSignal.GetType();
-						NewTraceEvent.TracePath = TracePath;
-						NewTraceEvent.InstanceName = InstanceName;
-						NewTraceEvent.FunctionName = CurrentFunctionName;
-						NewTraceEvent.NodeName = NodeEntry.ExecNode->GetName();
-						NewTraceEvent.Offset = CurrSignal.GetScriptCodeOffset();
-						NewTraceEvent.ObservationTime = CurrSignal.GetTime();
+						CurrentNodeData.ExclusiveEvents.Add(CurrSignal);
 						break;
 					}
 					case EScriptInstrumentation::NodeEntry:
 					case EScriptInstrumentation::NodeDebugSite:
 					{
 						// Handle timings for events called as functions
-						if (NodeEntry.ExecNode->IsCustomEvent())
+						if (CurrentNodeData.ImpureNode->IsCustomEvent())
 						{
 							// Ensure this is a different event
 							if (const UK2Node_Event* EventGraphNode = Cast<UK2Node_Event>(GraphNode))
@@ -1297,133 +1656,82 @@ bool FScriptEventPlayback::Process(const TArray<FScriptInstrumentedEvent>& Signa
 						}
 						// Process tunnel timings
 						ProcessTunnelEntryPoints(FunctionContext, CurrSignal);
-						// Queue events
-						NodeEntry.EntryTracePaths.Push(TracePath);
-						NodeEntry.UnproccessedEvents.Add(CurrSignal);
-						// Add Trace History
-						FBlueprintExecutionTrace& NewTraceEvent = BlueprintContext->AddNewTraceHistoryEvent();
-						NewTraceEvent.TraceType = CurrSignal.GetType();
-						NewTraceEvent.TracePath = TracePath;
-						NewTraceEvent.InstanceName = InstanceName;
-						NewTraceEvent.FunctionName = CurrentFunctionName;
-						NewTraceEvent.NodeName = NodeEntry.ExecNode->GetName();
-						NewTraceEvent.Offset = CurrSignal.GetScriptCodeOffset();
-						NewTraceEvent.ObservationTime = CurrSignal.GetTime();
+						CurrentNodeData.ExclusiveTracePaths.Push(TracePath);
+						CurrentNodeData.ExclusiveEvents.Add(CurrSignal);
+						AddToTraceHistory(CurrentNodeData.ImpureNode->GetName(), CurrSignal);
 						break;
 					}
 					case EScriptInstrumentation::PushState:
 					{
 						TraceStack.Push(TracePath);
-						NodeEntry.UnproccessedEvents.Add(CurrSignal);
-						break;
-					}
-					case EScriptInstrumentation::RestoreState:
-					{
-						check(TraceStack.Num());
-						TracePath = TraceStack.Last();
+						// Process execution sequence inclusive timing
 						if (GraphNode->IsA<UK2Node_ExecutionSequence>())
 						{
-							FScriptInstrumentedEvent OverrideEvent(CurrSignal);
-							OverrideEvent.OverrideType(EScriptInstrumentation::NodeEntry);
-							NodeEntry.UnproccessedEvents.Add(OverrideEvent);
-							NodeEntry.EntryTracePaths.Add(TracePath);
-							// Add Trace History
-							FBlueprintExecutionTrace& NewTraceEvent = BlueprintContext->AddNewTraceHistoryEvent();
-							NewTraceEvent.TraceType = OverrideEvent.GetType();
-							NewTraceEvent.TracePath = TracePath;
-							NewTraceEvent.InstanceName = InstanceName;
-							NewTraceEvent.FunctionName = CurrentFunctionName;
-							NewTraceEvent.NodeName = NodeEntry.ExecNode->GetName();
-							NewTraceEvent.Offset = CurrSignal.GetScriptCodeOffset();
-							NewTraceEvent.ObservationTime = CurrSignal.GetTime();
+							ProcessExecutionSequence(CurrentNodeData, CurrSignal);
 						}
 						break;
 					}
-					case EScriptInstrumentation::SuspendState:
-					{
-						ProcessingState = EEventProcessingResult::Suspended;
-						break;
-					}
 					case EScriptInstrumentation::PopState:
 					{
 						if (TraceStack.Num())
 						{
 							TracePath = TraceStack.Pop();
-						}
-						// Consolidate enclosed timings for execution sequence nodes
-						if (GraphNode->IsA<UK2Node_ExecutionSequence>())
-						{
-							FScriptInstrumentedEvent OverrideEvent(CurrSignal);
-							OverrideEvent.OverrideType(EScriptInstrumentation::NodeExit);
-							NodeEntry.UnproccessedEvents.Add(OverrideEvent);
-							int32 NodeEntryIdx = INDEX_NONE;
-							for (int32 NodeEventIdx = NodeEntry.UnproccessedEvents.Num()-2; NodeEventIdx >= 0; --NodeEventIdx)
+							// Process execution sequence inclusive timing
+							if (GraphNode->IsA<UK2Node_ExecutionSequence>())
 							{
-								const FScriptInstrumentedEvent& NodeEvent = NodeEntry.UnproccessedEvents[NodeEventIdx];
-
-								if (NodeEvent.GetType() == EScriptInstrumentation::NodeEntry || NodeEvent.GetType() == EScriptInstrumentation::NodeDebugSite)
-								{
-									if (NodeEntryIdx != INDEX_NONE)
-									{
-										// Keep only the first node entry.
-										NodeEntry.UnproccessedEvents.RemoveAt(NodeEventIdx, 1);
-									}
-									NodeEntryIdx = NodeEventIdx;
-								}
-								else if (NodeEvent.GetType() == EScriptInstrumentation::NodeExit)
-								{
-									// Strip out all node exits apart from the PopExecution overridden event.
-									NodeEntry.UnproccessedEvents.RemoveAt(NodeEventIdx, 1);
-								}
+								ProcessExecutionSequence(CurrentNodeData, CurrSignal);
 							}
 						}
 						break;
 					}
-					case EScriptInstrumentation::NodeExit:
+					case EScriptInstrumentation::RestoreState:
 					{
-						// Handle duplicate exit sites for If Then Else nodes, This should be cleaned up in the refactor.
-						if (GraphNode->IsA<UK2Node_IfThenElse>())
+						check(TraceStack.Num());
+						TracePath = TraceStack.Last();
+						// Process execution sequence inclusive timing
+						if (GraphNode->IsA<UK2Node_ExecutionSequence>())
 						{
-							const int32 NextSignal = SignalIdx + 1;
-							if (NextSignal < StopIdx && SignalData[NextSignal].GetType() == EScriptInstrumentation::NodeExit)
-							{
-								continue;
-							}
+							ProcessExecutionSequence(CurrentNodeData, CurrSignal);
 						}
+						break;
+					}
+					case EScriptInstrumentation::SuspendState:
+					{
+						ProcessingState = EEventProcessingResult::Suspended;
+						break;
+					}
+					case EScriptInstrumentation::NodeExit:
+					{
 						// Cleanup branching multiple exits and correct the tracepath
-						if (NodeEntry.UnproccessedEvents.Num() && NodeEntry.UnproccessedEvents.Last().GetType() == EScriptInstrumentation::NodeExit)
+						if (CurrentNodeData.ExclusiveEvents.Num() && CurrentNodeData.ExclusiveEvents.Last().GetType() == EScriptInstrumentation::NodeExit)
 						{
-							NodeEntry.UnproccessedEvents.Pop();
-							if (NodeEntry.EntryTracePaths.Num())
+							CurrentNodeData.ExclusiveEvents.Pop();
+							if (CurrentNodeData.ExclusiveTracePaths.Num())
 							{
-								TracePath = NodeEntry.EntryTracePaths.Last();
+								TracePath = CurrentNodeData.ExclusiveTracePaths.Last();
 							}
 						}
 						// Process tunnel timings
 						ProcessTunnelEntryPoints(FunctionContext, CurrSignal);
 						// Add Trace History
-						FBlueprintExecutionTrace& NewTraceEvent = BlueprintContext->AddNewTraceHistoryEvent();
-						NewTraceEvent.TraceType = CurrSignal.GetType();
-						NewTraceEvent.TracePath = TracePath;
-						NewTraceEvent.InstanceName = InstanceName;
-						NewTraceEvent.FunctionName = CurrentFunctionName;
-						NewTraceEvent.NodeName = NodeEntry.ExecNode->GetName();
-						NewTraceEvent.Offset = CurrSignal.GetScriptCodeOffset();
-						NewTraceEvent.ObservationTime = CurrSignal.GetTime();
+						AddToTraceHistory(CurrentNodeData.ImpureNode->GetName(), CurrSignal);
 						// Process node exit
 						const int32 ScriptCodeExit = CurrSignal.GetScriptCodeOffset();
-						TracePath.AddExitPin(ScriptCodeExit);
-						NodeEntry.UnproccessedEvents.Add(CurrSignal);
+						if (const UEdGraphPin* ValidPin = FunctionContext->GetPinFromCodeLocation(ScriptCodeExit))
+						{
+							TracePath.AddExitPin(ScriptCodeExit);
+						}
+						CurrentNodeData.ExclusiveEvents.Add(CurrSignal);
 						// Process cyclic linkage - reset to root link tracepath
-						TSharedPtr<FScriptExecutionNode> NextLink = NodeEntry.ExecNode->GetLinkedNodeByScriptOffset(ScriptCodeExit);
+						TSharedPtr<FScriptExecutionNode> NextLink = CurrentNodeData.ImpureNode->GetLinkedNodeByScriptOffset(ScriptCodeExit);
 						if (NextLink.IsValid() && NextLink->HasFlags(EScriptExecutionNodeFlags::CyclicLinkage))
 						{
 							const UEdGraphNode* LinkedGraphNode = Cast<UEdGraphNode>(NextLink->GetObservedObject());
-							if (PerNodeEventHelper* LinkedEntry = NodeMap.Find(LinkedGraphNode))
+							if (NodeSignalHelper* LinkedEntry = NodeMap.Find(LinkedGraphNode))
 							{
-								if (LinkedEntry->EntryTracePaths.Num())
+								if (LinkedEntry->ExclusiveTracePaths.Num())
 								{
-									TracePath = LinkedEntry->EntryTracePaths.Last();
+									TracePath = LinkedEntry->ExclusiveTracePaths.Last();
 								}
 							}
 						}
@@ -1431,7 +1739,7 @@ bool FScriptEventPlayback::Process(const TArray<FScriptInstrumentedEvent>& Signa
 					}
 					default:
 					{
-						NodeEntry.UnproccessedEvents.Add(CurrSignal);
+						CurrentNodeData.ExclusiveEvents.Add(CurrSignal);
 						break;
 					}
 				}
@@ -1453,26 +1761,32 @@ bool FScriptEventPlayback::Process(const TArray<FScriptInstrumentedEvent>& Signa
 			for (auto EventTiming : EventTimings)
 			{
 				FTracePath EventTracePath;
-				const FName EventFunctionName = BlueprintContext->GetEventFunctionName(EventTiming.Key);
-				FunctionContext = BlueprintContext->GetFunctionContext(EventFunctionName);
+				// The UCS, along with BP events declared as 'const' in native C++ code, are implemented as standalone functions, and not within the ubergraph function context.
+				FunctionContext = BlueprintContext->GetFunctionContext(EventTiming.Key);
+				if (!FunctionContext.IsValid())
+				{
+					FunctionContext = BlueprintContext->GetFunctionContextForEvent(EventTiming.Key);
+				}
 				check (FunctionContext.IsValid());
 				EventNode = FunctionContext->GetProfilerDataForNode(EventTiming.Key);
 				check (EventNode.IsValid());
-				TSharedPtr<FScriptPerfData> PerfData = EventNode->GetPerfDataByInstanceAndTracePath(InstanceName, EventTracePath);
+				TSharedPtr<FScriptPerfData> PerfData = EventNode->GetOrAddPerfDataByInstanceAndTracePath(InstanceName, EventTracePath);
 				PerfData->AddEventTiming(SignalData[StopIdx].GetTime() - SignalData[LastEventIdx].GetTime());
 			}
 			EventTimings.Reset();
 		}
 		// Process Node map timings -- this can probably be rolled back into submission during the initial processing and lose this extra iteration.
-		for (auto NodeEntry : NodeMap)
+		for (auto CurrentNodeData : NodeMap)
 		{
-			TSharedPtr<FScriptExecutionNode> ExecNode = NodeEntry.Value.ExecNode;
+			TSharedPtr<FScriptExecutionNode> ExecNode = CurrentNodeData.Value.ImpureNode;
 			TSharedPtr<FScriptExecutionNode> PureNode;
 			double PureNodeEntryTime = 0.0;
 			double PureChainEntryTime = 0.0;
 			double NodeEntryTime = 0.0;
-			int32 TracePathIdx = 0;
-			for (auto EventIter = NodeEntry.Value.UnproccessedEvents.CreateIterator(); EventIter; ++EventIter)
+			int32 ExclTracePathIdx = 0;
+			int32 InclTracePathIdx = 0;
+			// Process exclusive events for this node
+			for (auto EventIter = CurrentNodeData.Value.ExclusiveEvents.CreateIterator(); EventIter; ++EventIter)
 			{
 				switch(EventIter->GetType())
 				{
@@ -1484,14 +1798,13 @@ bool FScriptEventPlayback::Process(const TArray<FScriptInstrumentedEvent>& Signa
 						}
 						else if (PureNode.IsValid())
 						{
-							TSharedPtr<FScriptPerfData> PerfData = PureNode->GetPerfDataByInstanceAndTracePath(InstanceName, NodeEntry.Value.InputTracePaths.Pop());
+							TSharedPtr<FScriptPerfData> PerfData = PureNode->GetOrAddPerfDataByInstanceAndTracePath(InstanceName, CurrentNodeData.Value.InputTracePaths.Pop());
 							PerfData->AddEventTiming(EventIter->GetTime() - PureNodeEntryTime);
 						}
-
 						PureNode.Reset();
 						PureNodeEntryTime = EventIter->GetTime();
 						
-						FunctionContext = BlueprintContext->GetFunctionContext(EventIter->GetFunctionName());
+						FunctionContext = CurrentNodeData.Value.BlueprintContext->GetFunctionContext(EventIter->GetFunctionName());
 						check(FunctionContext.IsValid());
 
 						const int32 ScriptCodeOffset = EventIter->GetScriptCodeOffset();
@@ -1499,7 +1812,7 @@ bool FScriptEventPlayback::Process(const TArray<FScriptInstrumentedEvent>& Signa
 						{
 							if (FunctionContext->HasProfilerDataForNode(Pin->GetOwningNode()->GetFName()))
 							{
-								PureNode = NodeEntry.Value.FunctionContext->GetProfilerDataForNode(Pin->GetOwningNode()->GetFName());
+								PureNode = CurrentNodeData.Value.FunctionContext->GetProfilerDataForNode(Pin->GetOwningNode()->GetFName());
 							}
 						}
 						break;
@@ -1509,13 +1822,12 @@ bool FScriptEventPlayback::Process(const TArray<FScriptInstrumentedEvent>& Signa
 					{
 						if (PureNode.IsValid())
 						{
-							TSharedPtr<FScriptPerfData> PerfData = PureNode->GetPerfDataByInstanceAndTracePath(InstanceName, NodeEntry.Value.InputTracePaths.Pop());
+							TSharedPtr<FScriptPerfData> PerfData = PureNode->GetOrAddPerfDataByInstanceAndTracePath(InstanceName, CurrentNodeData.Value.InputTracePaths.Pop());
 							PerfData->AddEventTiming(EventIter->GetTime() - PureNodeEntryTime);
 
 							PureNode.Reset();
 							PureNodeEntryTime = 0.0;
 						}
-
 						if (NodeEntryTime == 0.0)
 						{
 							NodeEntryTime = EventIter->GetTime();
@@ -1525,28 +1837,108 @@ bool FScriptEventPlayback::Process(const TArray<FScriptInstrumentedEvent>& Signa
 					case EScriptInstrumentation::NodeExit:
 					{
 						check(NodeEntryTime != 0.0 );
-						const FTracePath NodeTracePath = NodeEntry.Value.EntryTracePaths[TracePathIdx++];
-						TSharedPtr<FScriptPerfData> PerfData = ExecNode->GetPerfDataByInstanceAndTracePath(InstanceName, NodeTracePath);
+						const FTracePath NodeTracePath = CurrentNodeData.Value.ExclusiveTracePaths[ExclTracePathIdx++];
+						TSharedPtr<FScriptPerfData> PerfData = ExecNode->GetOrAddPerfDataByInstanceAndTracePath(InstanceName, NodeTracePath);
 						double PureChainDuration = PureChainEntryTime != 0.0 ? (NodeEntryTime - PureChainEntryTime) : 0.0;
 						double NodeDuration = EventIter->GetTime() - NodeEntryTime;
 						PerfData->AddEventTiming(NodeDuration);
-						PureChainEntryTime = NodeEntryTime = 0.0;
-						PerfData->AddPureChainTiming(PureChainDuration);
+						PerfData->AddInclusiveTiming(PureChainDuration, false, false);
 						TSharedPtr<FScriptExecutionNode> PureChainNode = ExecNode->GetPureChainNode();
 						if (PureChainNode.IsValid())
 						{
-							PerfData = PureChainNode->GetPerfDataByInstanceAndTracePath(InstanceName, NodeTracePath);
+							PerfData = PureChainNode->GetOrAddPerfDataByInstanceAndTracePath(InstanceName, NodeTracePath);
 							PerfData->AddEventTiming(PureChainDuration);
 						}
+						PureChainEntryTime = NodeEntryTime = 0.0;
 						break;
 					}
 				}
 			}
+			// Process inclusive events for this node
+			if (CurrentNodeData.Value.InclusiveEvents.Num() > 0)
+			{
+				int32 EventDepth = 0;
+				TArray<FScriptInstrumentedEvent> ProcessQueue;
+				// Discard any nested timings, it's nasty that we have to do this extra processing.
+				for (auto EventIter : CurrentNodeData.Value.InclusiveEvents)
+				{
+					EventDepth = EventIter.IsNodeExit() ? (EventDepth-1) : EventDepth;
+					check (EventDepth>=0);
+					if (EventDepth == 0)
+					{
+						ProcessQueue.Add(EventIter);
+					}
+					EventDepth = EventIter.IsNodeEntry() ? (EventDepth+1) : EventDepth;
+				}
+				for (auto EventIter : ProcessQueue)
+				{
+					switch(EventIter.GetType())
+					{
+						case EScriptInstrumentation::NodeEntry:
+						case EScriptInstrumentation::NodeDebugSite:
+						{
+							if (NodeEntryTime == 0.0)
+							{
+								NodeEntryTime = EventIter.GetTime();
+							}
+							break;
+						}
+						case EScriptInstrumentation::NodeExit:
+						{
+							check(NodeEntryTime != 0.0 );
+							const FTracePath NodeTracePath = CurrentNodeData.Value.InclusiveTracePaths[InclTracePathIdx++];
+							TSharedPtr<FScriptPerfData> PerfData = ExecNode->GetOrAddPerfDataByInstanceAndTracePath(InstanceName, NodeTracePath);
+							double NodeDuration = EventIter.GetTime() - NodeEntryTime;
+							PerfData->AddInclusiveTiming(NodeDuration, false, true);
+							break;
+						}
+					}
+				}
+			}
 		}
 	}
 	return ProcessingState != EEventProcessingResult::Failed;
 }
 
+void FScriptEventPlayback::ProcessExecutionSequence(NodeSignalHelper& CurrentNodeData, const FScriptInstrumentedEvent& CurrSignal)
+{
+	switch(CurrSignal.GetType())
+	{
+		// For a sequence node a restore state represents the end of execution of a sequence pin's path (excluding the last pin)
+		case EScriptInstrumentation::RestoreState:
+		{
+			// Convert the restore state into a node exit signal
+			FScriptInstrumentedEvent OverrideEvent(CurrSignal);
+			OverrideEvent.OverrideType(EScriptInstrumentation::NodeEntry);
+			CurrentNodeData.ExclusiveEvents.Add(OverrideEvent);
+			CurrentNodeData.ExclusiveTracePaths.Add(TracePath);
+			// Add Trace History
+			AddToTraceHistory(CurrentNodeData.ImpureNode->GetName(), OverrideEvent);
+			break;
+		}
+		// For a sequence node a node entry represents the start of execution.
+		case EScriptInstrumentation::PushState:
+		{
+			// Convert the push state into a inclusive entry signal
+			FScriptInstrumentedEvent OverrideEvent(CurrSignal);
+			OverrideEvent.OverrideType(EScriptInstrumentation::NodeEntry);
+			CurrentNodeData.InclusiveEvents.Add(OverrideEvent);
+			CurrentNodeData.InclusiveTracePaths.Add(TracePath);
+			break;
+		}
+		// For a sequence node a restore state represents the end of execution of all the sequence pins.
+		case EScriptInstrumentation::PopState:
+		{
+			// Convert the pop state into a inclusive exit signal
+			FScriptInstrumentedEvent OverrideEvent(CurrSignal);
+			OverrideEvent.OverrideType(EScriptInstrumentation::NodeExit);
+			CurrentNodeData.InclusiveEvents.Add(OverrideEvent);
+			CurrentNodeData.InclusiveTracePaths.Add(TracePath);
+			break;
+		}
+	}
+}
+
 void FScriptEventPlayback::ProcessTunnelEntryPoints(TSharedPtr<FBlueprintFunctionContext> FunctionContext, const FScriptInstrumentedEvent& CurrSignal)
 {
 	const int32 ScriptOffset = CurrSignal.GetScriptCodeOffset();
@@ -1585,4 +1977,17 @@ void FScriptEventPlayback::ProcessTunnelEntryPoints(TSharedPtr<FBlueprintFunctio
 	}
 }
 
+void FScriptEventPlayback::AddToTraceHistory(const FName NodeName, const FScriptInstrumentedEvent& TraceSignal)
+{
+	// Add Trace History
+	FBlueprintExecutionTrace& NewTraceEvent = BlueprintContext->AddNewTraceHistoryEvent();
+	NewTraceEvent.TraceType = TraceSignal.GetType();
+	NewTraceEvent.TracePath = TracePath;
+	NewTraceEvent.InstanceName = InstanceName;
+	NewTraceEvent.FunctionName = CurrentFunctionName;
+	NewTraceEvent.NodeName = NodeName;
+	NewTraceEvent.Offset = TraceSignal.GetScriptCodeOffset();
+	NewTraceEvent.ObservationTime = TraceSignal.GetTime();
+}
+
 #undef LOCTEXT_NAMESPACE
diff --git a/Engine/Source/Developer/BlueprintProfiler/Public/BlueprintProfiler.h b/Engine/Source/Developer/BlueprintProfiler/Public/BlueprintProfiler.h
--- a/Engine/Source/Developer/BlueprintProfiler/Public/BlueprintProfiler.h
+++ b/Engine/Source/Developer/BlueprintProfiler/Public/BlueprintProfiler.h
@@ -34,7 +34,6 @@ public:
 	virtual FOnBPStatGraphLayoutChanged& GetGraphLayoutChangedDelegate() { return GraphLayoutChangedDelegate; }
 	virtual TSharedPtr<FBlueprintExecutionContext> GetBlueprintContext(const FString& BlueprintClassPath) override;
 	virtual TSharedPtr<FScriptExecutionNode> GetProfilerDataForNode(const UEdGraphNode* GraphNode) override;
-	virtual FName MapBlueprintInstance(TSharedPtr<FBlueprintExecutionContext> BlueprintContext, const FString& InstancePath) override;
 	virtual bool HasDataForInstance(const UObject* Instance) const override;
 	virtual void ProcessEventProfilingData() override;
 	virtual EBlueprintProfilerHeatMapDisplayMode::Type GetGraphNodeHeatMapDisplayMode() const { return GraphNodeHeatMapDisplayMode; }
diff --git a/Engine/Source/Developer/BlueprintProfiler/Public/BlueprintProfilerModule.h b/Engine/Source/Developer/BlueprintProfiler/Public/BlueprintProfilerModule.h
--- a/Engine/Source/Developer/BlueprintProfiler/Public/BlueprintProfilerModule.h
+++ b/Engine/Source/Developer/BlueprintProfiler/Public/BlueprintProfilerModule.h
@@ -14,7 +14,11 @@ namespace EBlueprintProfilerHeatMapDisplayMode
 		None = 0,
 		Inclusive,
 		Exclusive,
-		MaxTiming
+		MaxTiming,
+		Total,
+		HottestPath,
+		HottestEndpoint,
+		PinToPin
 	};
 }
 
@@ -53,9 +57,6 @@ public:
 	/** Returns the current profiling event data for node */
 	virtual TSharedPtr<class FScriptExecutionNode> GetProfilerDataForNode(const UEdGraphNode* GraphNode) = 0;
 
-	/** Maps a blueprint instance and links to blueprint profiling data, returns the correct instance name */
-	virtual FName MapBlueprintInstance(TSharedPtr<FBlueprintExecutionContext> BlueprintContext, const FString& InstancePath) = 0;
-
 	/** Returns if the profiler currently retains any data for the specified instance */
 	virtual bool HasDataForInstance(const UObject* Instance) const { return false; }
 
diff --git a/Engine/Source/Developer/BlueprintProfiler/Public/ScriptInstrumentationCapture.h b/Engine/Source/Developer/BlueprintProfiler/Public/ScriptInstrumentationCapture.h
--- a/Engine/Source/Developer/BlueprintProfiler/Public/ScriptInstrumentationCapture.h
+++ b/Engine/Source/Developer/BlueprintProfiler/Public/ScriptInstrumentationCapture.h
@@ -52,6 +52,9 @@ public:
 	/** Returns the system time when this event was generated */
 	double GetTime() const { return Time; }
 
+	/** Overrides the system time for this event */
+	void OverrideTime(const double OverrideTime) { Time = OverrideTime; }
+
 	/** Returns if this event represents a change in active class/blueprint */
 	bool IsNewClass() const { return EventType == EScriptInstrumentation::Class; }
 
@@ -67,6 +70,12 @@ public:
 	/** Returns if this event represents a node execution event */
 	bool IsNodeTiming() const {	return EventType == EScriptInstrumentation::NodeEntry || EventType == EScriptInstrumentation::NodeExit;	}
 
+	/** Returns if this event represents a node entry event */
+	bool IsNodeEntry() const { return EventType == EScriptInstrumentation::NodeEntry || EventType == EScriptInstrumentation::NodeDebugSite; }
+
+	/** Returns if this event represents a node exit event */
+	bool IsNodeExit() const { return EventType == EScriptInstrumentation::NodeExit; }
+
 protected:
 	
 	/** The event type */
diff --git a/Engine/Source/Developer/BlueprintProfiler/Public/ScriptInstrumentationPlayback.h b/Engine/Source/Developer/BlueprintProfiler/Public/ScriptInstrumentationPlayback.h
--- a/Engine/Source/Developer/BlueprintProfiler/Public/ScriptInstrumentationPlayback.h
+++ b/Engine/Source/Developer/BlueprintProfiler/Public/ScriptInstrumentationPlayback.h
@@ -26,7 +26,8 @@ class FBlueprintExecutionContext : public TSharedFromThis<FBlueprintExecutionCon
 public:
 
 	FBlueprintExecutionContext()
-		: ExecutionTraceHistory(256)
+		: bIsBlueprintMapped(false)
+		, ExecutionTraceHistory(256)
 	{
 	}
 
@@ -36,6 +37,12 @@ public:
 	/** Returns if the blueprint context is fully formed */
 	BLUEPRINTPROFILER_API bool IsContextValid() const { return Blueprint.IsValid() && BlueprintClass.IsValid(); }
 
+	/** Returns if the blueprint is mapped */
+	bool IsBlueprintMapped() const { return bIsBlueprintMapped; }
+
+	/** Removes the blueprint mapping */
+	void RemoveMapping();
+
 	/** Returns the blueprint this context represents */
 	TWeakObjectPtr<UBlueprint> GetBlueprint() const { return Blueprint; }
 
@@ -51,14 +58,20 @@ public:
 	/** Returns a new trace history event */
 	FBlueprintExecutionTrace& AddNewTraceHistoryEvent() { return ExecutionTraceHistory.WriteNewElementInitialized(); }
 
+	/** Returns if an event is mapped. */
+	bool IsEventMapped(const FName EventName) const;
+
 	/** Add event execution node */
-	void AddEventNode(TSharedPtr<FScriptExecutionNode> EventExecNode);
+	void AddEventNode(TSharedPtr<class FBlueprintFunctionContext> FunctionContext, TSharedPtr<FScriptExecutionNode> EventExecNode);
+
+	/** Register a function context for an event */
+	void RegisterEventContext(const FName EventName, TSharedPtr<FBlueprintFunctionContext> FunctionContext);
 
 	/** Returns true if an execution node exists for the instance */
 	BLUEPRINTPROFILER_API bool HasProfilerDataForInstance(const FName InstanceName) const;
 
-	/** Corrects the instance name if it is a PIE instance and returns the editor actor */
-	TWeakObjectPtr<const UObject> GetInstance(FName& InstanceNameInOut);
+	/** Maps the instance returning the system name for instance */
+	BLUEPRINTPROFILER_API FName MapBlueprintInstance(const FString& InstancePath);
 
 	/** Remaps PIE actor instance paths to editor actor instances */
 	BLUEPRINTPROFILER_API FName RemapInstancePath(const FName InstanceName) const ;
@@ -66,6 +79,9 @@ public:
 	/** Returns the function name containing the event */
 	FName GetEventFunctionName(const FName EventName) const;
 
+	/** Returns the function context containing the event */
+	TSharedPtr<class FBlueprintFunctionContext> GetFunctionContextForEvent(const FName EventName) const;
+
 	/** Returns the function context matching the supplied name providing it exists */
 	TSharedPtr<class FBlueprintFunctionContext> GetFunctionContext(const FName FunctionNameIn) const;
 
@@ -75,6 +91,12 @@ public:
 		FunctionContexts.Add(FunctionName) = FunctionContext; 
 	}
 
+	/** Returns if this context has an execution node for the specified graph pin */
+	BLUEPRINTPROFILER_API bool HasProfilerDataForPin(const UEdGraphPin* GraphPin) const;
+
+	/** Returns if this context has an execution node for the specified graph pin, otherwise it falls back to a owning node lookup */
+	BLUEPRINTPROFILER_API TSharedPtr<FScriptExecutionNode> GetProfilerDataForPin(const UEdGraphPin* GraphPin);
+
 	/** Returns if this context has an execution node for the specified graph node */
 	BLUEPRINTPROFILER_API bool HasProfilerDataForNode(const UEdGraphNode* GraphNode) const;
 
@@ -87,13 +109,18 @@ public:
 private:
 
 	/** Creates and returns a new function context, returns existing if present  */
-	TSharedPtr<FBlueprintFunctionContext> CreateFunctionContext(UEdGraph* Graph);
+	TSharedPtr<FBlueprintFunctionContext> CreateFunctionContext(const FName FunctionName, UEdGraph* Graph);
 
 	/** Walks the blueprint and maps events and functions ready for profiling data */
 	bool MapBlueprintExecution();
 
+	/** Corrects the instance name if it is a PIE instance and sets the object weakptr to the editor actor, returns true if this is a new instance */
+	bool ResolveInstance(FName& InstanceNameInOut, TWeakObjectPtr<const UObject>& ObjectInOut);
+
 private:
 
+	/** Indicates the blueprint has been mapped */
+	bool bIsBlueprintMapped;
 	/** Blueprint the context represents */
 	TWeakObjectPtr<UBlueprint> Blueprint;
 	/** Blueprint Generated class for the context */
@@ -104,10 +131,14 @@ private:
 	TSharedPtr<FScriptExecutionBlueprint> BlueprintNode;
 	/** UFunction contexts for this blueprint */
 	TMap<FName, TSharedPtr<FBlueprintFunctionContext>> FunctionContexts;
+	/** UFunction context lookup by event name */
+	TMap<FName, TSharedPtr<FBlueprintFunctionContext>> EventFunctionContexts;
 	/** PIE instance remap container */
 	TMap<FName, FName> PIEInstanceNameMap;
 	/** Editor instances */
-	TMap<FName, TWeakObjectPtr<const UObject>> ActorInstances;
+	TMap<FName, TWeakObjectPtr<const UObject>> EditorActorInstances;
+	/** PIE instances */
+	TMap<FName, TWeakObjectPtr<const UObject>> PIEActorInstances;
 	/** Event Trace History */
 	TSimpleRingBuffer<FBlueprintExecutionTrace> ExecutionTraceHistory;
 };
@@ -120,7 +151,7 @@ class FBlueprintFunctionContext : public TSharedFromThis<FBlueprintFunctionConte
 public:
 
 	/** Initialise the function context from the graph */
-	void InitialiseContextFromGraph(TSharedPtr<FBlueprintExecutionContext> BlueprintContext, UEdGraph* Graph);
+	void InitialiseContextFromGraph(TSharedPtr<FBlueprintExecutionContext> BlueprintContext, const FName FunctionNameIn, UEdGraph* Graph);
 
 	/** Map Function */
 	void MapFunction();
@@ -186,6 +217,9 @@ private:
 	/** Locates and generates event contexts for any input events passed in */
 	void CreateInputEvents(TSharedPtr<FBlueprintExecutionContext> BlueprintContextIn, const TArray<UK2Node*>& InputEventNodes);
 
+	/** Locates and generates event contexts for any async task events passed in */
+	void CreateAsyncTaskEvents(TSharedPtr<FBlueprintExecutionContext> BlueprintContextIn, const TArray<UK2Node_BaseAsyncTask*>& AsyncTaskNodes);
+
 	/** Processes and detects any cyclic links making the linkage safe for traversal */
 	bool DetectCyclicLinks(TSharedPtr<FScriptExecutionNode> ExecNode, TSet<TSharedPtr<FScriptExecutionNode>>& Filter);
 
@@ -200,6 +234,12 @@ private:
 	/** Maps pin execution */
 	void MapExecPins(TSharedPtr<FScriptExecutionNode> ExecNode, const TArray<UEdGraphPin*>& Pins);
 
+	/** Returns a valid pin name, creating one based on pin characteristics if not available */
+	FName GetPinName(const UEdGraphPin* Pin) const;
+
+	/** Returns a full qualified pin name including owning node id */
+	FName GetUniquePinName(const UEdGraphPin* Pin) const;
+
 	/** Maps execution pin tunnel entry */
 	TSharedPtr<FScriptExecutionNode> MapTunnelEntry(const FName TunnelName, const UEdGraphPin* ExecPin, const UEdGraphPin* TunnelPin, const int32 PinScriptOffset);
 
@@ -225,18 +265,22 @@ private:
 
 private:
 
+	/** The graph name */
+	FName GraphName;
 	/** The function name */
 	FName FunctionName;
 	/** The Owning Blueprint context */
 	TWeakPtr<FBlueprintExecutionContext> BlueprintContext;
 	/** The UFunction this context represents */
 	TWeakObjectPtr<UFunction> Function;
+	/** The blueprint that owns this function graph */
+	TWeakObjectPtr<UBlueprint> OwningBlueprint;
 	/** The blueprint generated class for the owning blueprint */
 	TWeakObjectPtr<UBlueprintGeneratedClass> BlueprintClass;
 	/** ScriptCode offset to node cache */
 	TMap<int32, TWeakObjectPtr<const UEdGraphNode>> ScriptOffsetToNodes;
 	/** ScriptCode offset to pin cache */
-	TMap<int32, TWeakObjectPtr<const UEdGraphPin>> ScriptOffsetToPins;
+	TMap<int32, FEdGraphPinReference> ScriptOffsetToPins;
 	/** Execution nodes containing profiling data */
 	TMap<FName, TSharedPtr<FScriptExecutionNode>> ExecutionNodes;
 	/** Graph entry points */
@@ -275,6 +319,18 @@ class FScriptEventPlayback : public TSharedFromThis<FScriptEventPlayback>
 {
 public:
 
+	struct NodeSignalHelper
+	{
+		TSharedPtr<FScriptExecutionNode> ImpureNode;
+		TSharedPtr<FBlueprintExecutionContext> BlueprintContext;
+		TSharedPtr<FBlueprintFunctionContext> FunctionContext;
+		TArray<class FScriptInstrumentedEvent> ExclusiveEvents;
+		TArray<class FScriptInstrumentedEvent> InclusiveEvents;
+		TArray<FTracePath> ExclusiveTracePaths;
+		TArray<FTracePath> InputTracePaths;
+		TArray<FTracePath> InclusiveTracePaths;
+	};
+
 	struct TunnelEventHelper
 	{
 		double TunnelEntryTime;
@@ -282,8 +338,21 @@ public:
 		TSharedPtr<FScriptExecutionTunnelEntry> TunnelEntryPoint;
 	};
 
+	struct PushState
+	{
+		PushState(const UEdGraphNode* NodeIn, const FTracePath& TracePathIn)
+			: InstigatingNode(NodeIn)
+			, TracePath(TracePathIn)
+		{
+		}
+
+		const UEdGraphNode* InstigatingNode;
+		FTracePath TracePath;
+	};
+
 	FScriptEventPlayback(TSharedPtr<FBlueprintExecutionContext> BlueprintContextIn, const FName InstanceNameIn)
 		: InstanceName(InstanceNameIn)
+		, CurrentFunctionName(NAME_None)
 		, ProcessingState(EEventProcessingResult::None)
 		, BlueprintContext(BlueprintContextIn)
 	{
@@ -303,13 +372,21 @@ public:
 
 private:
 
+	/** Process execution sequence */
+	void ProcessExecutionSequence(NodeSignalHelper& CurrentNodeData, const FScriptInstrumentedEvent& CurrSignal);
+
 	/** Process tunnel entry points */
 	void ProcessTunnelEntryPoints(TSharedPtr<FBlueprintFunctionContext> FunctionContext, const FScriptInstrumentedEvent& CurrSignal);
 
+	/** Add to trace history */
+	void AddToTraceHistory(const FName NodeName, const FScriptInstrumentedEvent& TraceSignal);
+
 private:
 
 	/** Instance name */
 	FName InstanceName;
+	/** Function name */
+	FName CurrentFunctionName;
 	/** Event name */
 	FName EventName;
 	/** Processing state */
diff --git a/Engine/Source/Developer/MessageLog/Private/UserInterface/SMessageLogMessageListRow.cpp b/Engine/Source/Developer/MessageLog/Private/UserInterface/SMessageLogMessageListRow.cpp
--- a/Engine/Source/Developer/MessageLog/Private/UserInterface/SMessageLogMessageListRow.cpp
+++ b/Engine/Source/Developer/MessageLog/Private/UserInterface/SMessageLogMessageListRow.cpp
@@ -172,7 +172,12 @@ void SMessageLogMessageListRow::CreateMessage(const TSharedRef<SHorizontalBox>&
 		RowContent = CreateHyperlink(InMessageToken, FText::FromString(URLToken->GetURL()));
 	}
 		break;
-
+	case EMessageToken::EdGraph:
+	{
+		IconBrushName = FName("PropertyWindow.Button_Browse");
+		RowContent = CreateHyperlink(InMessageToken, InMessageToken->ToText());
+	}
+	break;
 	case EMessageToken::Action:
 	{
 		const TSharedRef<FActionToken> ActionToken = StaticCastSharedRef<FActionToken>(InMessageToken);
diff --git a/Engine/Source/Editor/AIGraph/Private/AIGraph.cpp b/Engine/Source/Editor/AIGraph/Private/AIGraph.cpp
--- a/Engine/Source/Editor/AIGraph/Private/AIGraph.cpp
+++ b/Engine/Source/Editor/AIGraph/Private/AIGraph.cpp
@@ -136,7 +136,6 @@ void UAIGraph::CollectAllNodeInstances(TSet<UObject*>& NodeInstances)
 bool UAIGraph::CanRemoveNestedObject(UObject* TestObject) const
 {
 	return !TestObject->IsA(UEdGraphNode::StaticClass()) &&
-		!TestObject->IsA(UEdGraphPin::StaticClass()) &&
 		!TestObject->IsA(UEdGraph::StaticClass()) &&
 		!TestObject->IsA(UEdGraphSchema::StaticClass());
 }
diff --git a/Engine/Source/Editor/AIGraph/Private/AIGraphNode.cpp b/Engine/Source/Editor/AIGraph/Private/AIGraphNode.cpp
--- a/Engine/Source/Editor/AIGraph/Private/AIGraphNode.cpp
+++ b/Engine/Source/Editor/AIGraph/Private/AIGraphNode.cpp
@@ -68,6 +68,7 @@ void UAIGraphNode::PostEditImport()
 
 void UAIGraphNode::PostEditUndo()
 {
+	UEdGraphNode::PostEditUndo();
 	ResetNodeOwner();
 	
 	if (ParentNode)
diff --git a/Engine/Source/Editor/AnimGraph/Private/AnimGraphNode_BlendListBase.cpp b/Engine/Source/Editor/AnimGraph/Private/AnimGraphNode_BlendListBase.cpp
--- a/Engine/Source/Editor/AnimGraph/Private/AnimGraphNode_BlendListBase.cpp
+++ b/Engine/Source/Editor/AnimGraph/Private/AnimGraphNode_BlendListBase.cpp
@@ -79,6 +79,7 @@ void UAnimGraphNode_BlendListBase::RemovePinsFromOldPins(TArray<UEdGraphPin*>& O
 				// and if array index is greater than removed index, decrease index
 				if (ArrayIndex == RemovedArrayIndex)
 				{
+					OldPins[PinIdx]->MarkPendingKill();
 					OldPins.RemoveAt(PinIdx);
 					--PinIdx;
 				}
diff --git a/Engine/Source/Editor/AnimGraph/Private/AnimStateNodeBase.cpp b/Engine/Source/Editor/AnimGraph/Private/AnimStateNodeBase.cpp
--- a/Engine/Source/Editor/AnimGraph/Private/AnimStateNodeBase.cpp
+++ b/Engine/Source/Editor/AnimGraph/Private/AnimStateNodeBase.cpp
@@ -109,7 +109,7 @@ void UAnimStateNodeBase::PostLoad()
 	if(CustomVersion < FFrameworkObjectVersion::FixNonTransactionalPins)
 	{
 		int32 BrokenPinCount = 0;
-		for(UEdGraphPin* Pin : Pins)
+		for(UEdGraphPin_Deprecated* Pin : DeprecatedPins)
 		{
 			if(!Pin->HasAnyFlags(RF_Transactional))
 			{
diff --git a/Engine/Source/Editor/BehaviorTreeEditor/Private/BehaviorTreeDecoratorGraphNode_Logic.cpp b/Engine/Source/Editor/BehaviorTreeEditor/Private/BehaviorTreeDecoratorGraphNode_Logic.cpp
--- a/Engine/Source/Editor/BehaviorTreeEditor/Private/BehaviorTreeDecoratorGraphNode_Logic.cpp
+++ b/Engine/Source/Editor/BehaviorTreeEditor/Private/BehaviorTreeDecoratorGraphNode_Logic.cpp
@@ -76,7 +76,7 @@ UEdGraphPin* UBehaviorTreeDecoratorGraphNode_Logic::AddInputPin()
 
 void UBehaviorTreeDecoratorGraphNode_Logic::RemoveInputPin(class UEdGraphPin* Pin)
 {
-	Pin->BreakAllPinLinks();
+	Pin->MarkPendingKill();
 	Pins.Remove(Pin);
 }
 
diff --git a/Engine/Source/Editor/BehaviorTreeEditor/Private/BehaviorTreeGraph.cpp b/Engine/Source/Editor/BehaviorTreeEditor/Private/BehaviorTreeGraph.cpp
--- a/Engine/Source/Editor/BehaviorTreeEditor/Private/BehaviorTreeGraph.cpp
+++ b/Engine/Source/Editor/BehaviorTreeEditor/Private/BehaviorTreeGraph.cpp
@@ -143,6 +143,8 @@ void UBehaviorTreeGraph::UpdateAsset(int32 UpdateFlags)
 		}
 	}
 
+	// we can't look at pins until pin references have been fixed up post undo:
+	UEdGraphPin::ResolveAllPinReferences();
 	if (RootNode && RootNode->Pins.Num() > 0 && RootNode->Pins[0]->LinkedTo.Num() > 0)
 	{
 		UBehaviorTreeGraphNode* Node = Cast<UBehaviorTreeGraphNode>(RootNode->Pins[0]->LinkedTo[0]->GetOwningNode());
diff --git a/Engine/Source/Editor/BlueprintGraph/Classes/K2Node.h b/Engine/Source/Editor/BlueprintGraph/Classes/K2Node.h
--- a/Engine/Source/Editor/BlueprintGraph/Classes/K2Node.h
+++ b/Engine/Source/Editor/BlueprintGraph/Classes/K2Node.h
@@ -159,9 +159,6 @@ class UK2Node : public UEdGraphNode
 
 	// K2Node interface
 
-	UPROPERTY()
-	TMap<UEdGraphPin*, UEdGraphPin*> ByRefMatchupPins;
-
 	/**
 	 * Reallocate pins during reconstruction; by default ignores the old pins and calls AllocateDefaultPins()
 	 * If you override this to create additional pins you likely need to call RestoreSplitPins to restore any
diff --git a/Engine/Source/Editor/BlueprintGraph/Classes/K2Node_AddComponent.h b/Engine/Source/Editor/BlueprintGraph/Classes/K2Node_AddComponent.h
--- a/Engine/Source/Editor/BlueprintGraph/Classes/K2Node_AddComponent.h
+++ b/Engine/Source/Editor/BlueprintGraph/Classes/K2Node_AddComponent.h
@@ -18,6 +18,10 @@ class UK2Node_AddComponent : public UK2Node_CallFunction
 	UPROPERTY()
 	FString TemplateBlueprint;
 
+	//~ Begin UObject Interface
+	virtual void PostDuplicate(bool bDuplicateForPIE) override;
+	//~ End UObject Interface
+
 	//~ Begin UEdGraphNode Interface
 	virtual void AllocateDefaultPins() override;
 	virtual void DestroyNode() override;
@@ -56,6 +60,9 @@ class UK2Node_AddComponent : public UK2Node_CallFunction
 		return FindPinChecked(TEXT("bManualAttachment"));
 	}
 
+	/** Helper method used to instantiate a new component template after duplication. */
+	BLUEPRINTGRAPH_API void MakeNewComponentTemplate();
+
 	/** Static name of function to call */
 	static BLUEPRINTGRAPH_API FName GetAddComponentFunctionName();
 
diff --git a/Engine/Source/Editor/BlueprintGraph/Classes/K2Node_FunctionResult.h b/Engine/Source/Editor/BlueprintGraph/Classes/K2Node_FunctionResult.h
--- a/Engine/Source/Editor/BlueprintGraph/Classes/K2Node_FunctionResult.h
+++ b/Engine/Source/Editor/BlueprintGraph/Classes/K2Node_FunctionResult.h
@@ -27,6 +27,7 @@ class UK2Node_FunctionResult : public UK2Node_FunctionTerminator
 	virtual void PostPlacedNewNode();
 	virtual bool CanDuplicateNode() const { return true; }
 	virtual void PostPasteNode() override;
+	virtual bool CanUserDeleteNode() const override;
 	//~ End UK2Node Interface
 
 	//~ Begin UK2Node_EditablePinBase Interface
@@ -41,6 +42,7 @@ class UK2Node_FunctionResult : public UK2Node_FunctionTerminator
 	BLUEPRINTGRAPH_API TArray<UK2Node_FunctionResult*> GetAllResultNodes() const;
 
 protected:
+	void SyncWithEntryNode();
 	void SyncWithPrimaryResultNode();
 };
 
diff --git a/Engine/Source/Editor/BlueprintGraph/Classes/K2Node_MakeArray.h b/Engine/Source/Editor/BlueprintGraph/Classes/K2Node_MakeArray.h
--- a/Engine/Source/Editor/BlueprintGraph/Classes/K2Node_MakeArray.h
+++ b/Engine/Source/Editor/BlueprintGraph/Classes/K2Node_MakeArray.h
@@ -61,4 +61,6 @@ protected:
 
 	/** Returns the function to be used to add a function to the array */
 	BLUEPRINTGRAPH_API UFunction* GetArrayAddFunction() const;
+
+	void SyncPinNames();
 };
\ No newline at end of file
