Copying //UE4/Dev-Blueprints to //UE4/Dev-Main (Source: //UE4/Dev-Blueprints @ 3042808)

diff --git a/Engine/Source/Developer/BlueprintProfiler/Private/ScriptInstrumentationPlayback.cpp b/Engine/Source/Developer/BlueprintProfiler/Private/ScriptInstrumentationPlayback.cpp
--- a/Engine/Source/Developer/BlueprintProfiler/Private/ScriptInstrumentationPlayback.cpp
+++ b/Engine/Source/Developer/BlueprintProfiler/Private/ScriptInstrumentationPlayback.cpp
@@ -4,6 +4,7 @@
 #include "EditorStyleSet.h"
 #include "GraphEditorSettings.h"
 #include "BlueprintEditorUtils.h"
+#include "DelayAction.h"
 
 #define LOCTEXT_NAMESPACE "ScriptInstrumentationPlayback"
 
@@ -106,15 +107,29 @@ FName FBlueprintExecutionContext::MapBlueprintInstance(const FString& InstancePa
 		else
 		{
 			TSharedPtr<FScriptExecutionInstance> InstanceNode = StaticCastSharedPtr<FScriptExecutionInstance>(BlueprintNode->GetInstanceByName(InstanceName));
-			if (InstanceNode.IsValid() && PIEActorInstances.Contains(InstanceName))
+			if (InstanceNode.IsValid())
 			{
-				InstanceNode->SetPIEInstance(PIEActorInstances.Find(InstanceName)->Get());
+				if (const FWorldContext* PIEWorldContext = GEditor->GetPIEWorldContext())
+				{
+					PIEActorInstances.Add(InstanceName) = InstanceNode->GetActiveObject();
+				}
 			}
 		}
 	}
 	return InstanceName;
 }
 
+TSharedPtr<FScriptExecutionInstance> FBlueprintExecutionContext::GetInstanceExecNode(const FName InstanceName)
+{
+	TSharedPtr<FScriptExecutionInstance> Result;
+	if (BlueprintNode.IsValid())
+	{
+		FName RemappedInstanceName = RemapInstancePath(InstanceName);
+		Result = StaticCastSharedPtr<FScriptExecutionInstance>(BlueprintNode->GetInstanceByName(RemappedInstanceName));
+	}
+	return Result;
+}
+
 bool FBlueprintExecutionContext::HasProfilerDataForInstance(const FName InstanceName) const
 {
 	bool bHasInstanceData = false;
@@ -221,16 +236,16 @@ TSharedPtr<FBlueprintFunctionContext> FBlueprintExecutionContext::GetFunctionCon
 	return Result;
 }
 
-TSharedPtr<FBlueprintFunctionContext> FBlueprintExecutionContext::CreateFunctionContext(const FName FunctionName, UEdGraph* Graph)
+template<typename FunctionType> TSharedPtr<FunctionType> FBlueprintExecutionContext::CreateFunctionContext(const FName FunctionName, UEdGraph* Graph)
 {
 	TSharedPtr<FBlueprintFunctionContext>& Result = FunctionContexts.FindOrAdd(FunctionName);
 	if (!Result.IsValid())
 	{
-		Result = MakeShareable(new FBlueprintFunctionContext);
+		Result = MakeShareable(new FunctionType);
 		Result->InitialiseContextFromGraph(AsShared(), FunctionName, Graph);
 	}
 	check(Result.IsValid());
-	return Result;
+	return StaticCastSharedPtr<FunctionType>(Result);
 }
 
 bool FBlueprintExecutionContext::HasProfilerDataForPin(const UEdGraphPin* GraphPin) const
@@ -354,9 +369,29 @@ bool FBlueprintExecutionContext::MapBlueprintExecution()
 			}
 		}
 		// Create function stubs
+		TMap<UK2Node_Tunnel*, TSharedPtr<FBlueprintFunctionContext>> DiscoveredTunnels;
 		for (auto GraphEntry : FunctionGraphs)
 		{
-			CreateFunctionContext(GraphEntry.Key, GraphEntry.Value);
+			TSharedPtr<FBlueprintFunctionContext> NewFunctionContext = CreateFunctionContext<FBlueprintFunctionContext>(GraphEntry.Key, GraphEntry.Value);
+			NewFunctionContext->DiscoverTunnels(GraphEntry.Value, DiscoveredTunnels);
+		}
+		// Create and map each tunnels instance as its own function context.
+		TSet<UEdGraph*> MappedTunnelGraphs;
+		for (auto TunnelInstance : DiscoveredTunnels)
+		{
+			if (UEdGraph* TunnelGraph = TunnelInstance.Value->GetTunnelGraphFromNode(TunnelInstance.Key))
+			{
+				if (!MappedTunnelGraphs.Contains(TunnelGraph))
+				{
+					// Create the function context for the tunnel graph.
+					CreateFunctionContext<FBlueprintFunctionContext>(TunnelGraph->GetFName(), TunnelGraph);
+					MappedTunnelGraphs.Add(TunnelGraph);
+				}
+				// Create the tunnel instance context.
+				TSharedPtr<FBlueprintTunnelInstanceContext> NewTunnelContext = CreateFunctionContext<FBlueprintTunnelInstanceContext>(TunnelInstance.Key->GetFName(), TunnelGraph);
+				NewTunnelContext->MapTunnelContext(TunnelInstance.Value, TunnelInstance.Key);
+				TunnelInstance.Value->MapTunnelInstance(TunnelInstance.Key);
+			}
 		}
 		// Map the function context nodes
 		for (auto Context : FunctionContexts)
@@ -427,22 +462,9 @@ void FBlueprintFunctionContext::InitialiseContextFromGraph(TSharedPtr<FBlueprint
 						BlueprintContextIn->AddEventNode(AsShared(), EventExecNode);
 					}
 				}
-				// Create Input events
-				TArray<UK2Node*> InputEventNodes;
-				Graph->GetNodesOfClassEx<UK2Node_InputAction, UK2Node>(InputEventNodes);
-				Graph->GetNodesOfClassEx<UK2Node_InputKey, UK2Node>(InputEventNodes);
-				Graph->GetNodesOfClassEx<UK2Node_InputTouch, UK2Node>(InputEventNodes);
-				if (InputEventNodes.Num())
-				{
-					CreateInputEvents(BlueprintContextIn, InputEventNodes);
-				}
-				// Create async task events
-				TArray<UK2Node_BaseAsyncTask*> AsyncTaskNodes;
-				Graph->GetNodesOfClass<UK2Node_BaseAsyncTask>(AsyncTaskNodes);
-				if (AsyncTaskNodes.Num())
-				{
-					CreateAsyncTaskEvents(BlueprintContextIn, AsyncTaskNodes);
-				}
+				// Create any compiler generated events
+				FBlueprintDebugData& DebugData = BlueprintClass.Get()->GetDebugData();
+				CreateDelegatePinEvents(BlueprintContextIn, DebugData.GetCompilerGeneratedEvents());
 			}
 			else
 			{
@@ -483,28 +505,31 @@ void FBlueprintFunctionContext::InitialiseContextFromGraph(TSharedPtr<FBlueprint
 	}
 }
 
-void FBlueprintFunctionContext::MapFunction()
+void FBlueprintFunctionContext::DiscoverTunnels(UEdGraph* Graph, TMap<UK2Node_Tunnel*, TSharedPtr<FBlueprintFunctionContext>>& DiscoveredTunnels)
 {
-	// Map any macro instances
-	if (UBlueprint* BP = OwningBlueprint.Get())
+	if (Graph)
 	{
-		TArray<UEdGraph*> Graphs;
-		BP->GetAllGraphs(Graphs);
-		for (auto Graph : Graphs)
+		TArray<UK2Node_Tunnel*> GraphTunnels;
+		Graph->GetNodesOfClass<UK2Node_Tunnel>(GraphTunnels);
+		// Map sub graphs / composites and macros
+		for (auto Tunnel : GraphTunnels)
 		{
-			if (Graph->GetFName() == GraphName)
+			if (UK2Node_MacroInstance* MacroInstance = Cast<UK2Node_MacroInstance>(Tunnel))
 			{
-				TArray<UK2Node_Tunnel*> GraphTunnels;
-				Graph->GetNodesOfClass<UK2Node_Tunnel>(GraphTunnels);
-				// Map sub graphs / composites and macros
-				for (auto Tunnel : GraphTunnels)
-				{
-					MapTunnelInstance(Tunnel);
-				}
-				break;
+				DiscoverTunnels(MacroInstance->GetMacroGraph(), DiscoveredTunnels);
+				DiscoveredTunnels.Add(Tunnel, AsShared());
+			}
+			else if (UK2Node_Composite* CompositeInstance = Cast<UK2Node_Composite>(Tunnel))
+			{
+				DiscoverTunnels(CompositeInstance->BoundGraph, DiscoveredTunnels);
+				DiscoveredTunnels.Add(Tunnel, AsShared());
 			}
 		}
 	}
+}
+
+void FBlueprintFunctionContext::MapFunction()
+{
 	// Map all entry point execution paths
 	for (auto EntryPoint : EntryPoints)
 	{
@@ -516,6 +541,18 @@ void FBlueprintFunctionContext::MapFunction()
 				EntryPoint->AddChildNode(ExecutionNode);
 			}
 		}
+		else if (const UEdGraphPin* EntryPointPin = EntryPoint->GetObservedPin())
+		{
+			const int32 PinOffset = GetCodeLocationFromPin(EntryPointPin);
+			for (auto LinkedPin : EntryPointPin->LinkedTo)
+			{
+				TSharedPtr<FScriptExecutionNode> ExecutionNode = MapNodeExecution(LinkedPin->GetOwningNode());
+				if (ExecutionNode.IsValid())
+				{
+					EntryPoint->AddChildNode(ExecutionNode);
+				}
+			}
+		}
 		// Check for Cyclic Linkage
 		TSet<TSharedPtr<FScriptExecutionNode>> Filter;
 		DetectCyclicLinks(EntryPoint, Filter);
@@ -534,23 +571,37 @@ bool FBlueprintFunctionContext::DetectCyclicLinks(TSharedPtr<FScriptExecutionNod
 	}
 	else
 	{
-		if (ExecNode->IsCyclicNode())
+		if (!ExecNode->CanLinkAsCyclicNode())
 		{
 			Filter.Add(ExecNode);
 		}
-		for (auto Child : ExecNode->GetChildNodes())
+		for (TSharedPtr<FScriptExecutionNode>& Child : ExecNode->GetChildNodes())
 		{
-			DetectCyclicLinks(Child, Filter);
+			if (DetectCyclicLinks(Child, Filter))
+			{
+				// Replace child with cyclic link node.
+				FScriptExecNodeParams CycleLinkParams;
+				const FString NodeName = FString::Printf(TEXT("CyclicLinkTo_%i_%s"), ExecutionNodes.Num(), *Child->GetName().ToString());
+				CycleLinkParams.NodeName = FName(*NodeName);
+				CycleLinkParams.ObservedObject = Child->GetObservedObject();
+				CycleLinkParams.DisplayName = Child->GetDisplayName();
+				CycleLinkParams.Tooltip = LOCTEXT("CyclicLink_ToolTip", "Cyclic Link");
+				CycleLinkParams.NodeFlags = EScriptExecutionNodeFlags::CyclicLinkage|Child->GetFlags();
+				CycleLinkParams.IconColor = FLinearColor(0.91f, 0.16f, 0.16f);
+				const FSlateBrush* LinkIcon = Child->GetIcon();
+				CycleLinkParams.Icon = const_cast<FSlateBrush*>(LinkIcon);
+				Child = CreateExecutionNode(CycleLinkParams);
+			}
 		}
 		for (auto LinkedNode : ExecNode->GetLinkedNodes())
 		{
 			TSharedPtr<FScriptExecutionNode> LinkedExecNode = LinkedNode.Value;
-			if (LinkedExecNode->HasFlags(EScriptExecutionNodeFlags::FunctionTunnel))
+			if (LinkedExecNode->HasFlags(EScriptExecutionNodeFlags::TunnelInstance))
 			{
-				TSharedPtr<FScriptExecutionNode> TunnelBoundry = LinkedExecNode->GetLinkedNodeByScriptOffset(LinkedNode.Key);
-				if (TunnelBoundry.IsValid())
+				TSharedPtr<FScriptExecutionNode> TunnelBoundary = LinkedExecNode->GetLinkedNodeByScriptOffset(LinkedNode.Key);
+				if (TunnelBoundary.IsValid())
 				{
-					TSharedPtr<FScriptExecutionNode> TunnelExec = TunnelBoundry->GetLinkedNodeByScriptOffset(LinkedNode.Key);
+					TSharedPtr<FScriptExecutionNode> TunnelExec = TunnelBoundary->GetLinkedNodeByScriptOffset(LinkedNode.Key);
 					if (TunnelExec.IsValid())
 					{
 						ExecNode = TunnelExec;
@@ -579,196 +630,97 @@ bool FBlueprintFunctionContext::DetectCyclicLinks(TSharedPtr<FScriptExecutionNod
 	return false;
 }
 
-void FBlueprintFunctionContext::CreateInputEvents(TSharedPtr<FBlueprintExecutionContext> BlueprintContextIn, const TArray<UK2Node*>& InputEventNodes)
+void FBlueprintFunctionContext::CreateDelegatePinEvents(TSharedPtr<FBlueprintExecutionContext> BlueprintContextIn, const TMap<FName, FEdGraphPinReference>& PinEvents)
 {
-	TMap<const UEdGraphNode*, TArray<FName>> NodeEventDescs;
-	// Extract basic event info
-	if (UBlueprintGeneratedClass* BPClass = BlueprintClass.Get())
+	struct FPinDelegateDesc
 	{
-		for (auto EventNode : InputEventNodes)
+		FPinDelegateDesc(const FName EventNameIn, UEdGraphPin* DelegatePinIn)
+			: EventName(EventNameIn)
+			, DelegatePin(DelegatePinIn)
 		{
-			if (EventNode)
+		}
+
+		FName EventName;
+		UEdGraphPin* DelegatePin;
+	};
+	if (PinEvents.Num())
+	{
+		TMap<const UEdGraphNode*, TArray<FPinDelegateDesc>> NodeEventDescs;
+		// Build event contexts per node
+		for (auto PinEvent : PinEvents)
+		{
+			if (UEdGraphPin* DelegatePin = PinEvent.Value.Get())
 			{
-				TArray<FName>& Events = NodeEventDescs.FindOrAdd(EventNode);
-				if (UK2Node_InputAction* InputActionNode = Cast<UK2Node_InputAction>(EventNode))
-				{
-					const FString CustomEventName = FString::Printf(TEXT("InpActEvt_%s_%s"), *InputActionNode->InputActionName.ToString());
-					for (auto FunctionIter = TFieldIterator<UFunction>(BPClass); FunctionIter; ++FunctionIter)
-					{
-						if (FunctionIter->GetName().Contains(CustomEventName))
-						{
-							Events.Add(FunctionIter->GetFName());
-						}
-					}
-				}
-				else if (UK2Node_InputKey* InputKeyNode = Cast<UK2Node_InputKey>(EventNode))
-				{
-					const FName ModifierName = InputKeyNode->GetModifierName();
-					const FString CustomEventName = ModifierName != NAME_None ? FString::Printf(TEXT("InpActEvt_%s_%s"), *ModifierName.ToString(), *InputKeyNode->InputKey.ToString()) :
-																				FString::Printf(TEXT("InpActEvt_%s"), *InputKeyNode->InputKey.ToString());
-					for (auto FunctionIter = TFieldIterator<UFunction>(BPClass); FunctionIter; ++FunctionIter)
-					{
-						if (FunctionIter->GetName().Contains(CustomEventName))
-						{
-							Events.Add(FunctionIter->GetFName());
-						}
-					}
-				}
-				else if (UK2Node_InputTouch* InputTouchNode = Cast<UK2Node_InputTouch>(EventNode))
+				const UEdGraphNode* OwningNode = DelegatePin->GetOwningNode();
+				TArray<FPinDelegateDesc>& Events = NodeEventDescs.FindOrAdd(OwningNode);
+				Events.Add(FPinDelegateDesc(PinEvent.Key, DelegatePin));
+			}
+		}
+		// Generate the event exec nodes
+		for (auto NodeEvents : NodeEventDescs)
+		{
+			// Check if this node requires an event node creating.
+			bool bCreateEventNode = true;
+			for (auto Pin : NodeEvents.Key->Pins)
+			{
+				if (Pin->Direction == EGPD_Input && Pin->PinType.PinCategory == UEdGraphSchema_K2::PC_Exec)
 				{
-					struct EventPinData
-					{
-						EventPinData(UEdGraphPin* InPin,TEnumAsByte<EInputEvent> InEvent ){	Pin=InPin;EventType=InEvent;};
-						UEdGraphPin* Pin;
-						TEnumAsByte<EInputEvent> EventType;
-					};
-					TArray<UEdGraphPin*> ActivePins;
-					if (UEdGraphPin* InputTouchPressedPin = InputTouchNode->GetPressedPin())
-					{
-						if (InputTouchPressedPin->LinkedTo.Num() > 0)
-						{
-							ActivePins.Add(InputTouchPressedPin);
-						}
-					}
-					if (UEdGraphPin* InputTouchReleasedPin = InputTouchNode->GetReleasedPin())
-					{
-						if (InputTouchReleasedPin->LinkedTo.Num() > 0)
-						{
-							ActivePins.Add(InputTouchReleasedPin);
-						}
-					}
-					if (UEdGraphPin* InputTouchMovedPin = InputTouchNode->GetMovedPin())
-					{
-						if (InputTouchMovedPin->LinkedTo.Num() > 0)
-						{
-							ActivePins.Add(InputTouchMovedPin);
-						}
-					}
-					for (auto Pin : ActivePins)
-					{
-						const FString CustomEventName = FString::Printf(TEXT("InpTchEvt_%s"), *GetPinName(Pin).ToString());
-						for (auto FunctionIter = TFieldIterator<UFunction>(BPClass); FunctionIter; ++FunctionIter)
-						{
-							if (FunctionIter->GetName().Contains(CustomEventName))
-							{
-								Events.Add(FunctionIter->GetFName());
-							}
-						}
-					}
+					bCreateEventNode = false;
+					break;
 				}
 			}
-		}
-	}
-	// Build event contexts
-	for (auto NodeEvents : NodeEventDescs)
-	{
-		// Setup the basic exec node params.
-		FScriptExecNodeParams EventParams;
-		EventParams.ObservedObject = NodeEvents.Key;
-		EventParams.OwningGraphName = NAME_None;
-		EventParams.DisplayName = NodeEvents.Key->GetNodeTitle(ENodeTitleType::ListView);
-		EventParams.Tooltip = LOCTEXT("NavigateToEventLocationHyperlink_ToolTip", "Navigate to the Event");
-		EventParams.NodeFlags = EScriptExecutionNodeFlags::Event;
-		EventParams.IconColor = FLinearColor(0.91f, 0.16f, 0.16f);
-		const FSlateBrush* EventIcon = NodeEvents.Key->ShowPaletteIconOnNode() ? NodeEvents.Key->GetIconAndTint(EventParams.IconColor).GetOptionalIcon() :
-																				 FEditorStyle::GetBrush(TEXT("BlueprintProfiler.BPNode"));
-		EventParams.Icon = const_cast<FSlateBrush*>(EventIcon);
-		TSharedPtr<FScriptExecutionNode> EventExecNode;
-		// Create events for this node
-		for (FName EventName : NodeEvents.Value)
-		{
-			if (!EventExecNode.IsValid())
+			if (bCreateEventNode)
 			{
-				// Create the event node.
-				EventParams.NodeName = EventName;
-				EventExecNode = CreateExecutionNode(EventParams);
+				// Setup the basic exec node params.
+				FScriptExecNodeParams EventParams;
+				EventParams.NodeName = FName(*FString::Printf(TEXT("%s__DummyEvent"), *NodeEvents.Key->GetName()));
+				EventParams.ObservedObject = NodeEvents.Key;
+				EventParams.OwningGraphName = NAME_None;
+				EventParams.DisplayName = NodeEvents.Key->GetNodeTitle(ENodeTitleType::ListView);
+				EventParams.Tooltip = LOCTEXT("NavigateToEventLocationHyperlink_ToolTip", "Navigate to the Event");
+				EventParams.NodeFlags = EScriptExecutionNodeFlags::Event;
+				EventParams.IconColor = FLinearColor(0.91f, 0.16f, 0.16f);
+				const FSlateBrush* EventIcon = NodeEvents.Key->ShowPaletteIconOnNode() ? NodeEvents.Key->GetIconAndTint(EventParams.IconColor).GetOptionalIcon() :
+																						 FEditorStyle::GetBrush(TEXT("BlueprintProfiler.BPNode"));
+				EventParams.Icon = const_cast<FSlateBrush*>(EventIcon);
+				TSharedPtr<FScriptExecutionNode> EventExecNode = CreateExecutionNode(EventParams);
 				// Add entry points.
 				AddEntryPoint(EventExecNode);
 				BlueprintContextIn->AddEventNode(AsShared(), EventExecNode);
 			}
-			else
+			// Create the events for the pins
+			for (FPinDelegateDesc EventDesc : NodeEvents.Value)
 			{
-				// Register the exec node under the event name.
-				ExecutionNodes.Add(EventName) = EventExecNode;
+				FScriptExecNodeParams PinParams;
+				PinParams.NodeName = EventDesc.EventName;
+				PinParams.ObservedPin = EventDesc.DelegatePin;
+				PinParams.OwningGraphName = NAME_None;
+				PinParams.DisplayName = EventDesc.DelegatePin->GetDisplayName();
+				PinParams.Tooltip = LOCTEXT("ExecPin_ExpandExecutionPath_ToolTip", "Expand execution path");
+				PinParams.NodeFlags = EScriptExecutionNodeFlags::ExecPin|EScriptExecutionNodeFlags::EventPin;
+				PinParams.IconColor = FLinearColor(1.f, 1.f, 1.f, 0.8f);
+				const bool bPinLinked = EventDesc.DelegatePin->LinkedTo.Num() > 0;
+				const FSlateBrush* Icon = bPinLinked ?	FEditorStyle::GetBrush(TEXT("BlueprintProfiler.BPPinConnected")) : 
+														FEditorStyle::GetBrush(TEXT("BlueprintProfiler.BPPinDisconnected"));
+				PinParams.Icon = const_cast<FSlateBrush*>(Icon);
+				TSharedPtr<FScriptExecutionNode> PinExecNode = CreateExecutionNode(PinParams);
+				AddEntryPoint(PinExecNode);
 				// Register the function context as a handler for the event.
-				BlueprintContextIn->RegisterEventContext(EventName, AsShared());
+				BlueprintContextIn->RegisterEventContext(EventDesc.EventName, AsShared());
+				// Register exec node under pin name
+				const FName PinName = GetUniquePinName(EventDesc.DelegatePin);
+				ExecutionNodes.Add(PinName) = PinExecNode;
 			}
 		}
 	}
 }
 
-void FBlueprintFunctionContext::CreateAsyncTaskEvents(TSharedPtr<FBlueprintExecutionContext> BlueprintContextIn, const TArray<UK2Node_BaseAsyncTask*>& AsyncTaskNodes)
-{
-	struct FAsyncDelegateDesc
-	{
-		UEdGraphNode* DelegateNode;
-		UEdGraphPin* DelegatePin;
-		UFunction* EventFunction;
-	};
-	TArray<FAsyncDelegateDesc> AsyncDelegates;
-	// Extract basic event info
-	if (UBlueprintGeneratedClass* BPClass = BlueprintClass.Get())
-	{
-		FAsyncDelegateDesc NewDelegateDesc;
-		for (auto AsyncTaskNode : AsyncTaskNodes)
-		{
-			if (AsyncTaskNode)
-			{
-				for (auto Pin : AsyncTaskNode->Pins)
-				{
-					if (Pin->Direction == EGPD_Output && Pin->PinType.PinCategory == UEdGraphSchema_K2::PC_Exec)
-					{
-						FName PinName(GetPinName(Pin));
-						const FString PinDelegateName = FString::Printf(TEXT("%s_"), *PinName.ToString());
-						if (PinName != NAME_None)
-						{
-							for (auto FunctionIter = TFieldIterator<UFunction>(BPClass, EFieldIteratorFlags::ExcludeSuper); FunctionIter; ++FunctionIter)
-							{
-								if (FunctionIter->HasAllFunctionFlags(FUNC_BlueprintEvent|FUNC_BlueprintCallable) && FunctionIter->GetName().Contains(PinDelegateName))
-								{
-									NewDelegateDesc.DelegateNode = AsyncTaskNode;
-									NewDelegateDesc.DelegatePin = Pin;
-									NewDelegateDesc.EventFunction = *FunctionIter;
-									AsyncDelegates.Add(NewDelegateDesc);
-									break;
-								}
-							}
-						}
-					}
-				}
-			}
-		}
-	}
-	// Build event contexts
-	for (auto AsyncDelegateDesc : AsyncDelegates)
-	{
-		UEdGraphPin* DelegatePin = AsyncDelegateDesc.DelegatePin;
-		FName EventName = AsyncDelegateDesc.EventFunction->GetFName();
-		FScriptExecNodeParams EventParams;
-		EventParams.NodeName = EventName;
-		EventParams.ObservedObject = DelegatePin->GetOwningNode();
-		EventParams.ObservedPin = DelegatePin;
-		EventParams.OwningGraphName = NAME_None;
-		EventParams.DisplayName = DelegatePin->GetDisplayName();
-		EventParams.Tooltip = LOCTEXT("NavigateToEventLocationHyperlink_ToolTip", "Navigate to the Event");
-		EventParams.NodeFlags = EScriptExecutionNodeFlags::Event|EScriptExecutionNodeFlags::AsyncTaskDelegate;
-		EventParams.IconColor = FLinearColor(1.f, 1.f, 1.f, 0.8f);
-		const FSlateBrush* Icon = DelegatePin->LinkedTo.Num() ?	FEditorStyle::GetBrush(TEXT("BlueprintProfiler.BPPinConnected")) : 
-																FEditorStyle::GetBrush(TEXT("BlueprintProfiler.BPPinDisconnected"));
-		EventParams.Icon = const_cast<FSlateBrush*>(Icon);
-		TSharedPtr<FScriptExecutionNode> EventExecNode = CreateExecutionNode(EventParams);
-		const FName PinName = GetUniquePinName(DelegatePin);
-		ExecutionNodes.Add(PinName) = EventExecNode;
-		AddEntryPoint(EventExecNode);
-	}
-}
-
 TSharedPtr<FScriptExecutionNode> FBlueprintFunctionContext::MapNodeExecution(UEdGraphNode* NodeToMap)
 {
 	TSharedPtr<FScriptExecutionNode> MappedNode;
 	if (NodeToMap)
 	{
-		MappedNode = GetProfilerDataForNode(NodeToMap->GetFName());
+		MappedNode = GetProfilerDataForGraphNode(NodeToMap);
 		if (!MappedNode.IsValid())
 		{
 			FScriptExecNodeParams NodeParams;
@@ -806,11 +758,7 @@ TSharedPtr<FScriptExecutionNode> FBlueprintFunctionContext::MapNodeExecution(UEd
 				MappedNode->SetPureNodeScriptCodeRange(GetPureNodeScriptCodeRange(NodeToMap));
 			}
 			// Determine a few node characteristics to guide processing.
-			if (const UK2Node_BaseAsyncTask* AsyncNode = Cast<UK2Node_BaseAsyncTask>(NodeToMap))
-			{
-				MappedNode->AddFlags(ExecPins.Num() ? (EScriptExecutionNodeFlags::AsyncTask|EScriptExecutionNodeFlags::SequentialBranch) : EScriptExecutionNodeFlags::AsyncTask);
-			}
-			else if (ExecPins.Num() > 1)
+			if (ExecPins.Num() > 1)
 			{
 				const bool bExecSequence = NodeToMap->IsA<UK2Node_ExecutionSequence>();
 				MappedNode->AddFlags(bExecSequence ? EScriptExecutionNodeFlags::SequentialBranch : EScriptExecutionNodeFlags::ConditionalBranch);
@@ -840,10 +788,11 @@ TSharedPtr<FScriptExecutionNode> FBlueprintFunctionContext::MapNodeExecution(UEd
 					if (!bEventCall)
 					{
 						// Update the function context
-						TSharedPtr<FBlueprintFunctionContext> NewFunctionContext = BlueprintContext.Pin()->GetFunctionContext(CalledGraph->GetFName());
+						const FName FunctionGraphName = CalledGraph->GetFName();
+						TSharedPtr<FBlueprintFunctionContext> NewFunctionContext = BlueprintContext.Pin()->GetFunctionContext(FunctionGraphName);
 						if (NewFunctionContext.IsValid())
 						{
-							NewFunctionContext->AddCallSiteEntryPointsToNode(MappedNode);
+							ChildFunctionContexts.Add(FunctionGraphName) = NewFunctionContext;
 						}
 					}
 				}
@@ -909,7 +858,6 @@ void FBlueprintFunctionContext::MapInputPins(TSharedPtr<FScriptExecutionNode> Ex
 
 void FBlueprintFunctionContext::MapExecPins(TSharedPtr<FScriptExecutionNode> ExecNode, const TArray<UEdGraphPin*>& Pins)
 {
-	const bool bAsyncTask = ExecNode->HasFlags(EScriptExecutionNodeFlags::AsyncTask);
 	const bool bBranchedExecution = ExecNode->IsBranch();
 	int32 NumUnwiredPins = 0;
 	const UEdGraphPin* FinalExecPin = nullptr;
@@ -923,124 +871,96 @@ void FBlueprintFunctionContext::MapExecPins(TSharedPtr<FScriptExecutionNode> Exe
 		// Note: Pass through non-relevant (e.g. reroute) nodes here as they're not compiled and thus do not need to be mapped for profiling.
 		const UEdGraphPin* LinkedPin = Pin->LinkedTo.Num() > 0 ? FBlueprintEditorUtils::FindFirstCompilerRelevantLinkedPin(Pin->LinkedTo[0]) : nullptr;
 		UEdGraphNode* LinkedPinNode = LinkedPin ? LinkedPin->GetOwningNode() : nullptr;
-		const bool bTunnelBoundry = LinkedPinNode ? LinkedPinNode->IsA<UK2Node_Tunnel>() : false;
+		const bool bTunnelBoundary = LinkedPinNode ? LinkedPinNode->IsA<UK2Node_Tunnel>() : false;
+		// Try locate an already mapped pin.
+		if (TSharedPtr<FScriptExecutionNode>* MappedPin = ExecutionNodes.Find(PinName))
+		{
+			PinExecNode = *MappedPin;
+			ExecNode->AddLinkedNode(PinScriptCodeOffset, PinExecNode);
+			continue;
+		}
 		// Create any neccesary dummy pins for branches
 		if (bBranchedExecution)
 		{
-			if (bAsyncTask)
-			{
-				if (TSharedPtr<FScriptExecutionNode>* DelegateNode = ExecutionNodes.Find(PinName))
-				{
-					PinExecNode = *DelegateNode;
-					PinExecNode->AddFlags(EScriptExecutionNodeFlags::ExecPin);
-				}
-			}
-			if (!PinExecNode.IsValid())
-			{
-				FScriptExecNodeParams LinkNodeParams;
-				LinkNodeParams.NodeName = PinName;
-				LinkNodeParams.ObservedObject = Pin->GetOwningNode();
-				LinkNodeParams.DisplayName = Pin->GetDisplayName();
-				LinkNodeParams.ObservedPin = Pin;
-				LinkNodeParams.Tooltip = LOCTEXT("ExecPin_ExpandExecutionPath_ToolTip", "Expand execution path");
-				LinkNodeParams.NodeFlags = !bInvalidTrace ? EScriptExecutionNodeFlags::ExecPin : (EScriptExecutionNodeFlags::ExecPin|EScriptExecutionNodeFlags::InvalidTrace);
-				LinkNodeParams.IconColor = FLinearColor(1.f, 1.f, 1.f, 0.8f);
-				const FSlateBrush* Icon = LinkedPin ?	FEditorStyle::GetBrush(TEXT("BlueprintProfiler.BPPinConnected")) : 
-														FEditorStyle::GetBrush(TEXT("BlueprintProfiler.BPPinDisconnected"));
-				LinkNodeParams.Icon = const_cast<FSlateBrush*>(Icon);
-				PinExecNode = CreateExecutionNode(LinkNodeParams);
-			}
+			FScriptExecNodeParams LinkNodeParams;
+			LinkNodeParams.NodeName = PinName;
+			LinkNodeParams.ObservedObject = Pin->GetOwningNode();
+			LinkNodeParams.DisplayName = Pin->GetDisplayName();
+			LinkNodeParams.ObservedPin = Pin;
+			LinkNodeParams.Tooltip = LOCTEXT("ExecPin_ExpandExecutionPath_ToolTip", "Expand execution path");
+			LinkNodeParams.NodeFlags = !bInvalidTrace ? EScriptExecutionNodeFlags::ExecPin : (EScriptExecutionNodeFlags::ExecPin|EScriptExecutionNodeFlags::InvalidTrace);
+			LinkNodeParams.IconColor = FLinearColor(1.f, 1.f, 1.f, 0.8f);
+			const FSlateBrush* Icon = LinkedPin ?	FEditorStyle::GetBrush(TEXT("BlueprintProfiler.BPPinConnected")) : 
+													FEditorStyle::GetBrush(TEXT("BlueprintProfiler.BPPinDisconnected"));
+			LinkNodeParams.Icon = const_cast<FSlateBrush*>(Icon);
+			PinExecNode = CreateExecutionNode(LinkNodeParams);
 			ExecNode->AddLinkedNode(PinScriptCodeOffset, PinExecNode);
-			if (!bTunnelBoundry)
-			{
-				TSharedPtr<FScriptExecutionNode> LinkedPinExecNode = MapNodeExecution(LinkedPinNode);
-				if (LinkedPinExecNode.IsValid())
-				{
-					PinExecNode->AddChildNode(LinkedPinExecNode);
-				}
-			}
 		}
-		else
+		else if (!PinExecNode.IsValid())
 		{
-			if (bAsyncTask)
-			{
-				if (TSharedPtr<FScriptExecutionNode>* DelegateNode = ExecutionNodes.Find(PinName))
-				{
-					PinExecNode = *DelegateNode;
-				}
-			}
-			if (!PinExecNode.IsValid())
+			PinExecNode = ExecNode;
+		}
+		// Continue mapping forward.
+		TSharedPtr<FScriptExecutionNode> LinkedPinExecNode = bTunnelBoundary ? MapTunnelBoundary(LinkedPin) : MapNodeExecution(LinkedPinNode);
+		if (LinkedPinExecNode.IsValid())
+		{
+			if (bBranchedExecution)
 			{
-				PinExecNode = ExecNode;
+				PinExecNode->AddChildNode(LinkedPinExecNode);
 			}
-			if (!bTunnelBoundry)
+			else
 			{
-				TSharedPtr<FScriptExecutionNode> LinkedPinExecNode = MapNodeExecution(LinkedPinNode);
-				if (LinkedPinExecNode.IsValid())
-				{
-					PinExecNode->AddLinkedNode(PinScriptCodeOffset, LinkedPinExecNode);
-				}
+				PinExecNode->AddLinkedNode(PinScriptCodeOffset, LinkedPinExecNode);
 			}
 		}
-		// Handle tunnel boundries
-		if (bTunnelBoundry)
+	}
+}
+
+TSharedPtr<FScriptExecutionNode> FBlueprintFunctionContext::MapTunnelBoundary(const UEdGraphPin* TunnelPin)
+{
+	TSharedPtr<FScriptExecutionNode> TunnelBoundaryNode;
+	if (TunnelPin)
+	{
+		const FName TunnelBoundaryName = GetTunnelBoundaryName(TunnelPin);
+		TunnelBoundaryNode = GetProfilerDataForNodeChecked(TunnelBoundaryName);
+		if (TunnelBoundaryNode->IsTunnelEntry())
 		{
-			if (StagingTunnelName != NAME_None)
+			TSharedPtr<FScriptExecutionTunnelEntry> TunnelEntryInstance = StaticCastSharedPtr<FScriptExecutionTunnelEntry>(TunnelBoundaryNode);
+			UEdGraphNode* TunnelInstanceNode = TunnelPin->GetOwningNode();
+
+			for (auto ExitSite : TunnelEntryInstance->GetExitSites())
 			{
-				const FName TunnelExitName(GetPinName(LinkedPin));
-				TSharedPtr<FScriptExecutionNode> TunnelExitNode = GetProfilerDataForNode(TunnelExitName);
-				// Check for re-entrant exit site
-				if (ExecNode->HasFlags(EScriptExecutionNodeFlags::SequentialBranch))
+				if (ExitSite.Value->IsTunnelExit())
 				{
-					// Locate the last linked exec pin
-					if (!FinalExecPin)
+					TSharedPtr<FScriptExecutionTunnelExit> TunnelExit = TunnelEntryInstance->GetExitSite(ExitSite.Key);
+					const UEdGraphPin* TunnelInstanceExitPin = TunnelExit->GetExternalPin();
+					for (auto LinkedPin : TunnelInstanceExitPin->LinkedTo)
 					{
-						for (auto SeqPin : Pins)
+						UEdGraphNode* LinkedNode = LinkedPin->GetOwningNode();
+						TSharedPtr<FScriptExecutionNode> LinkedExecNode;
+						// Need to be careful here because a tunnel instance exit site can link to a tunnel Boundary too.
+						if (LinkedNode->IsA<UK2Node_Tunnel>())
 						{
-							if (SeqPin->LinkedTo.Num())
-							{
-								FinalExecPin = SeqPin;
-							}
+							LinkedExecNode = MapTunnelBoundary(LinkedPin);
 						}
-					}
-					// Last pin in sequence is not re-entrant, but all others are assumed to be.
-					TunnelExitNode->AddFlags(Pin == FinalExecPin ? EScriptExecutionNodeFlags::TunnelFinalExitPin : EScriptExecutionNodeFlags::ReEntrantTunnelPin);
-				}
-				if (bBranchedExecution)
-				{
-					PinExecNode->AddChildNode(TunnelExitNode);
-				}
-				else
-				{
-					const int32 TunnelId = GetTunnelIdFromName(TunnelExitName);
-					check (TunnelId != INDEX_NONE);
-					PinExecNode->AddFlags(EScriptExecutionNodeFlags::InvalidTrace);
-					PinExecNode->AddLinkedNode(TunnelId, TunnelExitNode);
-				}
-				TArray<TSharedPtr<FScriptExecutionNode>>& TunnelExitPoints = TunnelExitPointMap.FindOrAdd(StagingTunnelName);
-				TunnelExitPoints.Add(TunnelExitNode);
-			}
-			else
-			{
-				const FName TunnelName(GetPinName(LinkedPin));
-				TSharedPtr<FScriptExecutionNode> TunnelEntryNode = MapTunnelEntry(TunnelName, Pin, LinkedPin, PinScriptCodeOffset);
-				if (TunnelEntryNode.IsValid())
-				{
-					if (bBranchedExecution)
-					{
-						PinExecNode->AddChildNode(TunnelEntryNode);
-					}
-					else
-					{
-						PinExecNode->AddLinkedNode(PinScriptCodeOffset, TunnelEntryNode);
+						else
+						{
+							LinkedExecNode = GetProfilerDataForNode(LinkedNode->GetFName());
+						}
+						if (!LinkedExecNode.IsValid())
+						{
+							LinkedExecNode = MapNodeExecution(LinkedNode);
+						}
+						TunnelExit->AddLinkedNode(ExitSite.Key, LinkedExecNode);
 					}
 				}
 			}
 		}
 	}
+	return TunnelBoundaryNode;
 }
 
-FName FBlueprintFunctionContext::GetPinName(const UEdGraphPin* Pin) const
+FName FBlueprintFunctionContext::GetPinName(const UEdGraphPin* Pin)
 {
 	FName PinName(NAME_None);
 	if (Pin)
@@ -1058,13 +978,22 @@ FName FBlueprintFunctionContext::GetPinName(const UEdGraphPin* Pin) const
 			}
 			PinString = PinTypeIndex > 0 ? FString::Printf(TEXT("%s%i"), *Pin->PinType.PinCategory, PinTypeIndex) : Pin->PinType.PinCategory;
 		}
+		UEdGraphNode* OwningNode = Pin->GetOwningNode();
+		if (OwningNode->IsA<UK2Node_Tunnel>())
+		{
+			// Tunnel pins have to be unique so we need the node name in addition to pin name.
+			UEdGraph* TunnelGraph = GetTunnelGraphFromNode(OwningNode);
+			TunnelGraph = TunnelGraph ? TunnelGraph : OwningNode->GetTypedOuter<UEdGraph>();
+			FName TunnelName = TunnelGraph->GetFName();
+			PinString = FString::Printf(TEXT("%s_%s"), *TunnelName.ToString(), *PinString);
+		}
 		PinName = FName(*PinString);
 	}
 	check (PinName != NAME_None);
 	return PinName;
 }
 
-FName FBlueprintFunctionContext::GetUniquePinName(const UEdGraphPin* Pin) const
+FName FBlueprintFunctionContext::GetUniquePinName(const UEdGraphPin* Pin)
 {
 	FName PinName = GetPinName(Pin);
 	if (Pin)
@@ -1075,134 +1004,81 @@ FName FBlueprintFunctionContext::GetUniquePinName(const UEdGraphPin* Pin) const
 	return PinName;
 }
 
-TSharedPtr<FScriptExecutionNode> FBlueprintFunctionContext::MapTunnelEntry(const FName TunnelName, const UEdGraphPin* ExecPin, const UEdGraphPin* TunnelPin, const int32 PinScriptOffset)
+FName FBlueprintFunctionContext::GetTunnelBoundaryName(const UEdGraphPin* Pin)
 {
-	// Retrieve tunnel function context
-	const UK2Node_Tunnel* TunnelNode = Cast<UK2Node_Tunnel>(TunnelPin->GetOwningNode());
-	TSharedPtr<FBlueprintFunctionContext> TunnelContext = GetTunnelContextFromNode(TunnelNode);
-	const int32 TunnelId = TunnelContext->GetTunnelIdFromName(TunnelName);
-	check(TunnelId != INDEX_NONE);
-	// Create a tunnel instance node for the tunnel/macro to centralise stats within and aid stat lookup - for external lookup purposes by node name.
-	TSharedPtr<FScriptExecutionTunnelInstance> TunnelInstance = GetTypedProfilerDataForNode<FScriptExecutionTunnelInstance>(TunnelNode->GetFName());
-	if (!TunnelInstance.IsValid())
-	{
-		FScriptExecNodeParams TunnelInstanceParams;
-		TunnelInstanceParams.NodeName = TunnelNode->GetFName();
-		TunnelInstanceParams.ObservedObject = TunnelNode;
-		TunnelInstanceParams.DisplayName = TunnelNode->GetNodeTitle(ENodeTitleType::ListView);
-		TunnelInstanceParams.Tooltip = LOCTEXT("NavigateToNodeLocationHyperlink_ToolTip", "Navigate to the Node");
-		TunnelInstanceParams.NodeFlags = EScriptExecutionNodeFlags::FunctionTunnel;
-		TunnelInstanceParams.IconColor = FLinearColor(1.f, 1.f, 1.f, 0.8f);
-		const FSlateBrush* TunnelIcon = TunnelNode->ShowPaletteIconOnNode() ? TunnelNode->GetIconAndTint(TunnelInstanceParams.IconColor).GetOptionalIcon() :
-																			  FEditorStyle::GetBrush(TEXT("BlueprintProfiler.BPNode"));
-		TunnelInstanceParams.Icon = const_cast<FSlateBrush*>(TunnelIcon);
-		TunnelInstance = CreateTypedExecutionNode<FScriptExecutionTunnelInstance>(TunnelInstanceParams);
-	}
-	// Create a customised tunnel entry nodes,
-	// This maps each unique entry point per tunnel instance so we can only include the entry and exit sites in use.
-	FScriptExecNodeParams TunnelParams;
-	TunnelParams.NodeName = FName(*FString::Printf(TEXT("%s_EntryPointId%i"), *TunnelNode->GetFName().ToString(), TunnelId));
-	TunnelParams.ObservedObject = TunnelNode;
-	TunnelParams.DisplayName = TunnelNode->GetNodeTitle(ENodeTitleType::ListView);
-	TunnelParams.Tooltip = LOCTEXT("NavigateToNodeLocationHyperlink_ToolTip", "Navigate to the Node");
-	TunnelParams.NodeFlags = EScriptExecutionNodeFlags::FunctionTunnel;
-	TunnelParams.IconColor = FLinearColor(1.f, 1.f, 1.f, 0.8f);
-	const FSlateBrush* TunnelIcon = TunnelNode->ShowPaletteIconOnNode() ? TunnelNode->GetIconAndTint(TunnelParams.IconColor).GetOptionalIcon() :
-																		  FEditorStyle::GetBrush(TEXT("BlueprintProfiler.BPNode"));
-	TunnelParams.Icon = const_cast<FSlateBrush*>(TunnelIcon);
-	FScriptExecutionTunnelEntry::ETunnelType TunnelType = TunnelContext->GetTunnelType(TunnelName);
-	TSharedPtr<FScriptExecutionTunnelEntry> CustomTunnelEntryNode = MakeShareable(new FScriptExecutionTunnelEntry(TunnelParams, TunnelNode, TunnelType));
-	ExecutionNodes.FindOrAdd(TunnelParams.NodeName) = CustomTunnelEntryNode;
-	// Add the custom entry point to the tunnel instance ( for stat updates )
-	TunnelInstance->AddCustomEntryPoint(CustomTunnelEntryNode);
-	// Add the active entry point ( pin node ) to the custom tunnel entry node
-	TSharedPtr<FScriptExecutionNode> EntryPinNode = TunnelContext->GetProfilerDataForNode(TunnelName);
-	check (EntryPinNode.IsValid());
-	CustomTunnelEntryNode->AddChildNode(EntryPinNode);
-	if (!EntryPinNode->GetNumChildren())
-	{
-		// Map the linked sites to the entry point on first encounter. ( this should probably be done during macro mapping )
-		TSharedPtr<FScriptExecutionNode> FirstTunnelNode = TunnelContext->GetTunnelEntryPoint(TunnelName);
-		check(FirstTunnelNode.IsValid());
-		EntryPinNode->AddChildNode(FirstTunnelNode);
-	}
-	// Map tunnel exit points
-	TArray<TSharedPtr<FScriptExecutionNode>> TunnelExitPoints;
-	TunnelContext->GetTunnelExitPoints(TunnelName, TunnelExitPoints);
-	// Build a pin name to exit pin map ( located on the tunnel node rather than in the tunnel graph )
-	TMap<FName, const UEdGraphPin*> NameToPinMap;
-	for (auto ExitPin : TunnelNode->Pins)
+	FName PinName(NAME_None);
+	if (Pin)
 	{
-		if (ExitPin->PinType.PinCategory == UEdGraphSchema_K2::PC_Exec && ExitPin->LinkedTo.Num())
-		{
-			NameToPinMap.Add(GetPinName(ExitPin)) = ExitPin;
-		}
+		UEdGraphNode* OwningNode = Pin->GetOwningNode();
+		const FString PinString = Pin->PinName.IsEmpty() ? Pin->PinType.PinCategory : Pin->PinName;
+		PinName = FName(*FString::Printf(TEXT("%s_%s"), *OwningNode->GetFName().ToString(), *PinString));
 	}
-	// Register tunnel input pin node using the script offset from the first script offset emitted pin.
-	// This is to try and improve tunnel entry detection for stats.
-	if (const UEdGraphPin** InputPin = NameToPinMap.Find(TunnelName))
+	check (PinName != NAME_None);
+	return PinName;
+}
+
+UEdGraphPin* FBlueprintFunctionContext::FindMatchingPin(const UEdGraphNode* NodeToSearch, const UEdGraphPin* PinToFind, const bool bIgnoreDirection)
+{
+	UEdGraphPin* MatchingPin = nullptr;
+	if (NodeToSearch && PinToFind)
 	{
-		if ((*InputPin)->LinkedTo.Num())
+		// Quick simple lookup, works the bulk of the time.
+		MatchingPin = NodeToSearch->FindPin(PinToFind->PinName);
+		// More exhaustive search, for Boundary cases
+		if (!MatchingPin)
 		{
-			if (const UEdGraphPin* InputLinkedPin = (*InputPin)->LinkedTo[0])
+			const FName ForcedPinName = GetPinName(PinToFind);
+			for (auto SearchPin : NodeToSearch->Pins)
 			{
-				const int32 EntryScriptCodeOffset = GetCodeLocationFromPin(InputLinkedPin);
-				ScriptOffsetToPins.Add(EntryScriptCodeOffset) = *InputPin;
-				TunnelEntrySites.Add(EntryScriptCodeOffset) = CustomTunnelEntryNode;
+				if (PinToFind->PinName == SearchPin->PinName && PinToFind->PinType.PinCategory == SearchPin->PinType.PinCategory)
+				{
+					if (bIgnoreDirection || PinToFind->Direction == SearchPin->Direction)
+					{
+						MatchingPin = SearchPin;
+						break;
+					}
+				}
 			}
 		}
 	}
-	// Map the tunnel exit sites to the exit sites for this instance and determine script offsets.
-	for (auto CurrExitPoint : TunnelExitPoints)
+	return MatchingPin;
+}
+
+UEdGraph* FBlueprintFunctionContext::GetTunnelGraphFromNode(const UEdGraphNode* TunnelNode)
+{
+	UEdGraph* TunnelGraph = nullptr;
+	if (const UK2Node_MacroInstance* MacroInstance = Cast<UK2Node_MacroInstance>(TunnelNode))
+	{
+		TunnelGraph = MacroInstance->GetMacroGraph();
+	}
+	else if (const UK2Node_Composite* CompositeInstance = Cast<UK2Node_Composite>(TunnelNode))
 	{
-		if (const UEdGraphPin** ExitPin = NameToPinMap.Find(CurrExitPoint->GetName()))
+		TunnelGraph = CompositeInstance->BoundGraph;
+	}
+	return TunnelGraph;
+}
+
+bool FBlueprintFunctionContext::IsTunnelInternal(const UEdGraphNode* TunnelNode)
+{
+	return !(TunnelNode->IsA<UK2Node_MacroInstance>()||TunnelNode->IsA<UK2Node_Composite>());
+}
+
+void FBlueprintFunctionContext::MapTunnelExits(TSharedPtr<FScriptExecutionTunnelEntry> TunnelEntryPoint)
+{
+	if (TunnelEntryPoint.IsValid())
+	{
+		for (auto ExitSite : TunnelEntryPoint->GetExitSites())
 		{
-			for (auto LinkedExitPin : (*ExitPin)->LinkedTo)
+			if (const UEdGraphPin* ExitPin = ExitSite.Value->GetObservedPin())
 			{
-				// This is the unique script offset for this instance's exit site.
-				const int32 ExitScriptCodeOffset = GetCodeLocationFromPin(*ExitPin);
-				check (ExitScriptCodeOffset != INDEX_NONE);
-				// Add entry site for exit script offset - not sure what this is used for now - if anything.
-				TunnelExitSites.Add(ExitScriptCodeOffset) = CustomTunnelEntryNode;
-				// Pass through any non-relevant (e.g. reroute) nodes.
-				LinkedExitPin = FBlueprintEditorUtils::FindFirstCompilerRelevantLinkedPin(LinkedExitPin);
-				// Continue mapping exit sites
-				UEdGraphNode* LinkedExitNode = LinkedExitPin->GetOwningNode();
-				TSharedPtr<FScriptExecutionNode> LinkedExecNode;
-				if (LinkedExitNode->IsA<UK2Node_Tunnel>())
+				for (UEdGraphPin* LinkedPin : ExitPin->LinkedTo)
 				{
-					// Map being wired into another tunnel.
-					const FName LinkedTunnelName(GetPinName(LinkedExitPin));
-					LinkedExecNode = MapTunnelEntry(LinkedTunnelName, *ExitPin, LinkedExitPin, ExitScriptCodeOffset);
-				}
-				else
-				{
-					// Map forward normally.
-					LinkedExecNode = MapNodeExecution(LinkedExitPin->GetOwningNode());
-				}
-				if (LinkedExecNode.IsValid())
-				{
-					if (CurrExitPoint->HasFlags(EScriptExecutionNodeFlags::ReEntrantTunnelPin))
-					{
-						// Add the linked nodes into the CurExitPoint
-						CurrExitPoint->AddLinkedNode(ExitScriptCodeOffset, LinkedExecNode);
-					}
-					else
-					{
-						// Add the linked nodes into the CurExitPoint and as a linked node to the custom tunnel entry.
-						CurrExitPoint->AddFlags(EScriptExecutionNodeFlags::TunnelFinalExitPin);
-						LinkedExecNode->AddFlags(EScriptExecutionNodeFlags::TunnelFinalExitPin);
-						CustomTunnelEntryNode->AddLinkedNode(ExitScriptCodeOffset, LinkedExecNode);
-						CurrExitPoint->AddLinkedNode(ExitScriptCodeOffset, LinkedExecNode);
-					}
-					// Notify the custom tunnel entry node that this is a valid exit site for this instance.
-					CustomTunnelEntryNode->AddExitSite(CurrExitPoint, ExitScriptCodeOffset);
+					TSharedPtr<FScriptExecutionNode> LinkedExecNode = MapNodeExecution(LinkedPin->GetOwningNode());
+					ExitSite.Value->AddChildNode(LinkedExecNode);
 				}
 			}
 		}
 	}
-	// Return the tunnel entry.
-	return CustomTunnelEntryNode;
 }
 
 bool FBlueprintFunctionContext::HasProfilerDataForNode(const FName NodeName) const
@@ -1217,15 +1093,80 @@ TSharedPtr<FScriptExecutionNode> FBlueprintFunctionContext::GetProfilerDataForNo
 	{
 		Result = *SearchResult;
 	}
+	else
+	{
+		// Check child function contexts
+		for (auto ChildFunctionContext : ChildFunctionContexts)
+		{
+			if (ChildFunctionContext.Value.IsValid())
+			{
+				Result = ChildFunctionContext.Value.Pin()->GetProfilerDataForNode(NodeName);
+				if (Result.IsValid())
+				{
+					break;
+				}
+			}
+		}
+	}
+	return Result;
+}
+
+TSharedPtr<FScriptExecutionNode> FBlueprintFunctionContext::GetProfilerDataForNodeChecked(const FName NodeName)
+{
+	TSharedPtr<FScriptExecutionNode> Result = GetProfilerDataForNode(NodeName);
+	check (Result.IsValid());
+	return Result;
+}
+
+TSharedPtr<FScriptExecutionNode> FBlueprintFunctionContext::GetProfilerDataForGraphNode(const UEdGraphNode* Node)
+{
+	TSharedPtr<FScriptExecutionNode> Result;
+	if (Node)
+	{
+		if (UEdGraph* OuterGraph = Node->GetTypedOuter<UEdGraph>())
+		{
+			const bool bTunnelInstance = Node->IsA<UK2Node_MacroInstance>()||Node->IsA<UK2Node_Composite>();
+			// Correctly locate macro/tunnel graphs
+			const FName ObjectGraphName = bTunnelInstance ? Node->GetFName() : OuterGraph->GetFName();
+			if (ObjectGraphName == GraphName)
+			{
+				Result = GetProfilerDataForNode(Node->GetFName());
+			}
+			else if (TWeakPtr<FBlueprintFunctionContext>* FunctionContext = ChildFunctionContexts.Find(ObjectGraphName))
+			{
+				if (FunctionContext->IsValid())
+				{
+					Result = FunctionContext->Pin()->GetProfilerDataForNode(Node->GetFName());
+				}
+			}
+		}
+	}
+	if (Result.IsValid() && Result->GetObservedObject() != Node)
+	{
+		// Objects don't match, perform a more exhaustive search.
+		Result.Reset();
+	}
 	return Result;
 }
 
 template<typename ExecNodeType> TSharedPtr<ExecNodeType> FBlueprintFunctionContext::GetTypedProfilerDataForNode(const FName NodeName)
 {
 	TSharedPtr<ExecNodeType> Result;
-	if (TSharedPtr<FScriptExecutionNode>* SearchResult = ExecutionNodes.Find(NodeName))
+	TSharedPtr<FScriptExecutionNode> SearchResult = GetProfilerDataForNode(NodeName);
+	if (SearchResult.IsValid())
+	{
+		Result = StaticCastSharedPtr<ExecNodeType>(SearchResult);
+	}
+	return Result;
+}
+
+template<typename ExecNodeType> TSharedPtr<ExecNodeType> FBlueprintFunctionContext::GetTypedProfilerDataForGraphNode(const UEdGraphNode* Node)
+{
+	TSharedPtr<ExecNodeType> Result;
+	TSharedPtr<FScriptExecutionNode> SearchResult = GetProfilerDataForGraphNode(Node);
+	if (SearchResult.IsValid())
 	{
-		Result = StaticCastSharedPtr<ExecNodeType>(*SearchResult);
+		Result = StaticCastSharedPtr<ExecNodeType>(SearchResult);
 	}
 	return Result;
 }
@@ -1247,6 +1188,19 @@ const UEdGraphNode* FBlueprintFunctionContext::GetNodeFromCodeLocation(const int
 	return Result.Get();
 }
 
+const UEdGraphNode* FBlueprintFunctionContext::FindTunnelSourceNodeFromCodeLocation(const int32 ScriptOffset)
+{
+	TWeakObjectPtr<const UEdGraphNode> Result;
+	if (!Result.IsValid() && BlueprintClass.IsValid())
+	{
+		if (const UEdGraphNode* GraphNode = BlueprintClass->GetDebugData().FindMacroSourceNodeFromCodeLocation(Function.Get(), ScriptOffset))
+		{
+			Result = GraphNode;
+		}
+	}
+	return Result.Get();
+}
+
 const UEdGraphPin* FBlueprintFunctionContext::GetPinFromCodeLocation(const int32 ScriptOffset)
 {
 	FEdGraphPinReference& Result = ScriptOffsetToPins.FindOrAdd(ScriptOffset);
@@ -1310,6 +1264,33 @@ template<typename ExecNodeType> TSharedPtr<ExecNodeType> FBlueprintFunctionConte
 	return NewTypedNode;
 }
 
+void FBlueprintFunctionContext::MapTunnelInstance(UK2Node_Tunnel* TunnelInstance)
+{
+	// Grab tunnel context from instance name
+	const FName TunnelGraphName = TunnelInstance->GetFName();
+	TSharedPtr<FBlueprintTunnelInstanceContext> TunnelGraphContext = StaticCastSharedPtr<FBlueprintTunnelInstanceContext>(BlueprintContext.Pin()->GetFunctionContext(TunnelGraphName));
+	if (TunnelGraphContext.IsValid())
+	{
+		// Add the instance context
+		ChildFunctionContexts.Add(TunnelGraphName) = TunnelGraphContext;
+		// Add any child contexts belonging to the instance.
+		for (auto ChildContexts : TunnelGraphContext->ChildFunctionContexts)
+		{
+			ChildFunctionContexts.Add(ChildContexts.Key) = ChildContexts.Value;
+		}
+	}
+}
+
+void FBlueprintFunctionContext::AddEntryPoint(TSharedPtr<FScriptExecutionNode> EntryPoint)
+{
+	EntryPoints.Add(EntryPoint);
+}
+
+void FBlueprintFunctionContext::AddExitPoint(TSharedPtr<FScriptExecutionNode> ExitPoint)
+{
+	ExitPoints.Add(ExitPoint);
+}
+
 void FBlueprintFunctionContext::AddCallSiteEntryPointsToNode(TSharedPtr<FScriptExecutionNode> CallingNode) const
 {
 	if (CallingNode.IsValid())
@@ -1324,237 +1305,322 @@ void FBlueprintFunctionContext::AddCallSiteEntryPointsToNode(TSharedPtr<FScriptE
 	}
 }
 
-void FBlueprintFunctionContext::MapTunnelInstance(UK2Node_Tunnel* TunnelInstance)
+//////////////////////////////////////////////////////////////////////
+// FBlueprintTunnelInstanceContext
+
+void FBlueprintTunnelInstanceContext::InitialiseContextFromGraph(TSharedPtr<FBlueprintExecutionContext> BlueprintContextIn, const FName TunnelInstanceName, UEdGraph* TunnelGraph)
 {
-	check(TunnelInstance);
-	UEdGraph* TunnelGraph = nullptr;
-	if (UK2Node_MacroInstance* MacroInstance = Cast<UK2Node_MacroInstance>(TunnelInstance))
-	{
-		TunnelGraph = MacroInstance->GetMacroGraph();
-	}
-	else if (UK2Node_Composite* CompositeInstance = Cast<UK2Node_Composite>(TunnelInstance))
-	{
-		TunnelGraph = CompositeInstance->BoundGraph;
-	}
-	TSharedPtr<FBlueprintFunctionContext> TunnelGraphContext;
-	if (TunnelGraph)
+	GraphName = TunnelGraph->GetFName();
+	FunctionName = TunnelInstanceName;
+	BlueprintContext = BlueprintContextIn;
+}
+
+void FBlueprintTunnelInstanceContext::MapTunnelContext(TSharedPtr<FBlueprintFunctionContext> CallingFunctionContext, UK2Node_Tunnel* TunnelInstance)
+{
+	if (!TunnelInstanceNode.IsValid())
 	{
-		TunnelGraphContext = BlueprintContext.Pin()->GetFunctionContext(TunnelGraph->GetFName());
-		if (!TunnelGraphContext.IsValid())
-		{
-			// Drill down for any nested macros
-			TArray<UK2Node_Tunnel*> GraphTunnels;
-			TunnelGraph->GetNodesOfClass<UK2Node_Tunnel>(GraphTunnels);
-			// Map sub graphs / composites and macros
-			for (auto ChildTunnel : GraphTunnels)
+		// Set the mapping context.
+		Function = CallingFunctionContext->GetUFunction();
+		BlueprintClass = CallingFunctionContext->GetBlueprintClass();
+		TunnelInstanceNode = TunnelInstance;
+		// Find the function context that represents the graph and not the instance
+		UEdGraph* TunnelGraph = GetTunnelGraphFromNode(TunnelInstance);
+		TunnelFunctionContext = BlueprintContext.Pin()->GetFunctionContext(TunnelGraph->GetFName());
+		check (TunnelFunctionContext.IsValid());
+		// Map tunnel Input/Output
+		MapTunnelIO();
+		// Create tunnel instance node.
+		FScriptExecNodeParams TunnelInstanceParams;
+		TunnelInstanceParams.NodeName = TunnelInstance->GetFName();
+		TunnelInstanceParams.ObservedObject = TunnelInstance;
+		TunnelInstanceParams.DisplayName = TunnelInstance->GetNodeTitle(ENodeTitleType::ListView);
+		TunnelInstanceParams.Tooltip = LOCTEXT("NavigateToNodeLocationHyperlink_ToolTip", "Navigate to the Node");
+		TunnelInstanceParams.NodeFlags = EScriptExecutionNodeFlags::TunnelInstance;
+		TunnelInstanceParams.IconColor = FLinearColor(1.f, 1.f, 1.f, 0.8f);
+		if (TunnelInstance->IsA<UK2Node_Composite>())
+		{
+			const FSlateBrush* TunnelIcon = FEditorStyle::GetBrush(TEXT("GraphEditor.SubGraph_16x"));
+			TunnelInstanceParams.Icon = const_cast<FSlateBrush*>(TunnelIcon);
+		}
+		else
+		{
+			const FSlateBrush* TunnelIcon = TunnelInstance->ShowPaletteIconOnNode() ? TunnelInstance->GetIconAndTint(TunnelInstanceParams.IconColor).GetOptionalIcon() :
+																					  FEditorStyle::GetBrush(TEXT("GraphEditor.SubGraph_16x"));
+			TunnelInstanceParams.Icon = const_cast<FSlateBrush*>(TunnelIcon);
+		}
+		TSharedPtr<FScriptExecutionTunnelInstance> TunnelInstanceExecNode = CreateTypedExecutionNode<FScriptExecutionTunnelInstance>(TunnelInstanceParams);
+		ExternalNodes.Add(TunnelInstance->GetFName()) = TunnelInstanceExecNode;
+		// Map child tunnel instances now, because we require that the instance exec node is setup.
+		TMap<UK2Node_Tunnel*, TSharedPtr<FBlueprintFunctionContext>> ChildTunnels;
+		DiscoverTunnels(TunnelGraph, ChildTunnels);
+		for (auto ChildTunnel : ChildTunnels)
+		{
+			MapTunnelInstance(ChildTunnel.Key);
+		}
+		// Add Tunnel Graph Context as child context
+		ChildFunctionContexts.Add(TunnelFunctionContext->GetFunctionName()) = TunnelFunctionContext;
+		// Find tunnel instance entry sites
+		TMap<FName, UEdGraphPin*> InstanceEntrySites;
+		for (auto Pin : TunnelInstance->Pins)
+		{
+			if (Pin->PinType.PinCategory == UEdGraphSchema_K2::PC_Exec)
 			{
-				UEdGraph* ChildTunnelGraph = nullptr;
-				if (UK2Node_MacroInstance* MacroInstance = Cast<UK2Node_MacroInstance>(ChildTunnel))
-				{
-					ChildTunnelGraph = MacroInstance->GetMacroGraph();
-				}
-				else if (UK2Node_Composite* CompositeInstance = Cast<UK2Node_Composite>(ChildTunnel))
+				if (Pin->LinkedTo.Num())
 				{
-					ChildTunnelGraph = CompositeInstance->BoundGraph;
-				}
-				if (ChildTunnelGraph && ChildTunnelGraph != TunnelGraph)
-				{
-					MapTunnelInstance(ChildTunnel);
+					const FName EntryPointName = GetPinName(Pin);
+					if (Pin->Direction == EGPD_Input)
+					{
+						InstanceEntrySites.Add(EntryPointName) = Pin;
+					}
 				}
 			}
-			// Create the sub graph context
-			TunnelGraphContext = MakeShareable(new FBlueprintFunctionContext);
-			TunnelGraphContext->GraphName = TunnelGraph->GetFName();
-			TunnelGraphContext->FunctionName = TunnelGraph->GetFName();
-			TunnelGraphContext->BlueprintContext = BlueprintContext;
-			TunnelGraphContext->Function = Function;
-			TunnelGraphContext->BlueprintClass = BlueprintClass;
-			BlueprintContext.Pin()->AddFunctionContext(TunnelGraph->GetFName(), TunnelGraphContext);
-			// Map tunnel Inputs
-			TunnelGraphContext->MapTunnelIO(TunnelGraph);
-			// Map tunneling pins into the tunnel graph
-			for (auto EntryPoint : TunnelGraphContext->EntryPoints)
+		}
+		// Map tunnel entry site to exit sites
+		for (auto InstanceEntryPoint : InstanceEntrySites)
+		{
+			const int32 EntryPointScriptOffset = GetCodeLocationFromPin(InstanceEntryPoint.Value);
+			TSharedPtr<FScriptExecutionNode> EntryPoint = GetProfilerDataForNode(InstanceEntryPoint.Key);
+			if (EntryPoint.IsValid() && EntryPointScriptOffset != INDEX_NONE)
 			{
-				TArray<TSharedPtr<FScriptExecutionNode>>& MappedExits = TunnelGraphContext->TunnelExitPointMap.FindOrAdd(EntryPoint->GetName());
-				if (const UEdGraphPin* EntryPointPin = EntryPoint->GetObservedPin())
+				// Create custom exec node for each tunnel instance entry, this is so we can customise the appearance of the node.
+				FName BoundaryName = GetTunnelBoundaryName(InstanceEntryPoint.Value);
+				FScriptExecNodeParams TunnelEntryParams;
+				TunnelEntryParams.NodeName = BoundaryName;
+				TunnelEntryParams.ObservedObject = TunnelInstance;
+				TunnelEntryParams.ObservedPin = InstanceEntryPoint.Value;
+				TunnelEntryParams.DisplayName = TunnelInstance->GetNodeTitle(ENodeTitleType::ListView);
+				TunnelEntryParams.Tooltip = LOCTEXT("NavigateToNodeLocationHyperlink_ToolTip", "Navigate to the Node");
+				TunnelEntryParams.NodeFlags = EScriptExecutionNodeFlags::TunnelEntryPinInstance;
+				TunnelEntryParams.IconColor = FLinearColor(1.f, 1.f, 1.f, 0.8f);
+				TunnelEntryParams.Icon = TunnelInstanceParams.Icon;
+				TSharedPtr<FScriptExecutionTunnelEntry> InstanceTunnelEntry = CreateTypedExecutionNode<FScriptExecutionTunnelEntry>(TunnelEntryParams);
+				// Add internal entrypoint as child.
+				InstanceTunnelEntry->AddChildNode(EntryPoint);
+				// Register the external entry point
+				TunnelInstanceExecNode->AddEntrySite(EntryPointScriptOffset, InstanceTunnelEntry);
+				ExternalNodes.Add(BoundaryName) = InstanceTunnelEntry;
+				// Set the current entry point we are mapping.
+				StagingEntryPoint = InstanceTunnelEntry;
+				// Map this tunnel entry
+				if (const UEdGraphPin* InternalEntryPointPin = EntryPoint->GetObservedPin())
 				{
-					TunnelGraphContext->StagingTunnelName = EntryPoint->GetName();
-					for (auto LinkedPin : EntryPointPin->LinkedTo)
+					for (auto LinkedPin : InternalEntryPointPin->LinkedTo)
 					{
 						// Pass through any non-relevant (e.g. reroute) nodes.
 						LinkedPin = FBlueprintEditorUtils::FindFirstCompilerRelevantLinkedPin(LinkedPin);
-
 						const UEdGraphNode* LinkedNode = LinkedPin->GetOwningNode();
-						if (LinkedNode->IsA<UK2Node_Tunnel>())
+						const UK2Node_Tunnel* LinkedTunnelNode = Cast<UK2Node_Tunnel>(LinkedNode);
+						const UEdGraph* ChildTunnelGraph = GetTunnelGraphFromNode(LinkedTunnelNode);
+						const FName ChildGraphName = ChildTunnelGraph ? ChildTunnelGraph->GetFName() : NAME_None;
+						const bool bLinkedToTunnelExit = LinkedTunnelNode && ChildGraphName == GraphName;
+						const int32 ScriptCodeOffset = GetCodeLocationFromPin(LinkedPin);
+						if (bLinkedToTunnelExit)
 						{
-							// This handles input linked directly to output
-							TSharedPtr<FScriptExecutionNode> OutputPinNode = TunnelGraphContext->GetProfilerDataForNode(GetPinName(LinkedPin));
+							// This handles tunnel input linked directly to tunnel output as a special case.
+							TSharedPtr<FScriptExecutionTunnelExit> OutputPinNode = GetTypedProfilerDataForNode<FScriptExecutionTunnelExit>(GetPinName(LinkedPin));
 							check(OutputPinNode.IsValid());
-							MappedExits.Add(OutputPinNode);
+							StagingEntryPoint->AddExitSite(ScriptCodeOffset, OutputPinNode);
+							EntryPoint->AddLinkedNode(ScriptCodeOffset, OutputPinNode);
 						}
 						else
 						{
 							// This adds any encountered exit points to the active tunnel name.
-							TSharedPtr<FScriptExecutionNode> ExecNode = TunnelGraphContext->MapNodeExecution(LinkedPin->GetOwningNode());
+							const bool bLinkedToTunnel = ChildGraphName != NAME_None && ChildGraphName != GraphName;
+							TSharedPtr<FScriptExecutionNode> ExecNode = bLinkedToTunnel ? MapTunnelBoundary(LinkedPin) : MapNodeExecution(LinkedPin->GetOwningNode());
 							if (ExecNode.IsValid())
 							{
-								TunnelGraphContext->TunnelEntryPointMap.Add(EntryPoint->GetName()) = ExecNode;
+								EntryPoint->AddLinkedNode(ScriptCodeOffset, ExecNode);
 							}
 						}
 					}
 				}
-			}
-			TunnelGraphContext->StagingTunnelName = NAME_None;
-		}
-		// Setup tunnel graph context from instance.
-		TunnelGraphContext->BlueprintContext = BlueprintContext;
-		TunnelGraphContext->Function = Function;
-		TunnelGraphContext->BlueprintClass = BlueprintClass;
-		// Patch info into the calling function context
-		for (auto ExecutionNode : TunnelGraphContext->ExecutionNodes)
-		{
-			const int32 TunnelId = TunnelGraphContext->GetTunnelIdFromName(ExecutionNode.Value->GetName());
-			if (TunnelId == INDEX_NONE)
-			{
-				TSharedPtr<FScriptExecutionNode>* ExistingNode = ExecutionNodes.Find(ExecutionNode.Key);
-				if (ExistingNode == nullptr)
+				// Add exit sites to instance
+				for (auto ExitSite : InstanceTunnelEntry->GetExitSites())
 				{
-					ExecutionNodes.Add(ExecutionNode.Key) = ExecutionNode.Value;
+					if (ExitSite.Value->IsTunnelExit())
+					{
+						TunnelInstanceExecNode->AddExitSite(ExitSite.Key, StaticCastSharedPtr<FScriptExecutionTunnelExit>(ExitSite.Value));
+					}
 				}
+				// Reset staged entry point
+				StagingEntryPoint.Reset();
 			}
 		}
+		// Reset the mapping context
+		Function.Reset();
+		BlueprintClass.Reset();
 	}
 }
 
-const FScriptExecutionTunnelEntry::ETunnelType FBlueprintFunctionContext::GetTunnelType(const FName TunnelName) const
+TSharedPtr<FScriptExecutionNode> FBlueprintTunnelInstanceContext::MapNodeExecution(UEdGraphNode* NodeToMap)
 {
-	FScriptExecutionTunnelEntry::ETunnelType TunnelTypeResult = TunnelEntryPointMap.Num() > 1 ? FScriptExecutionTunnelEntry::SinglePathTunnel : FScriptExecutionTunnelEntry::SimpleTunnel;
-	if (const TArray<TSharedPtr<FScriptExecutionNode>>* ExitSites = TunnelExitPointMap.Find(TunnelName))
+	TSharedPtr<FScriptExecutionNode> MappedNode;
+	if (NodeToMap)
 	{
-		TunnelTypeResult = ExitSites->Num() > 1 ? FScriptExecutionTunnelEntry::MultiIOTunnel : TunnelTypeResult;
+		MappedNode = GetProfilerDataForGraphNode(NodeToMap);
+		if (!MappedNode.IsValid())
+		{
+			// First encounter, map execution.
+			MappedNode = FBlueprintFunctionContext::MapNodeExecution(NodeToMap);
+		}
+		else
+		{
+			// Discover already mapped exit sites
+			DiscoverExitSites(MappedNode);
+		}
 	}
-	return TunnelTypeResult;
+	return MappedNode;
 }
 
-TSharedPtr<FBlueprintFunctionContext> FBlueprintFunctionContext::GetTunnelContextFromNode(const UEdGraphNode* TunnelNode) const
+bool FBlueprintTunnelInstanceContext::IsPinFromThisTunnel(const UEdGraphPin* TunnelPin) const
 {
-	TSharedPtr<FBlueprintFunctionContext> TunnelContext;
-	if (const UK2Node_Composite* CompositeNode = Cast<UK2Node_Composite>(TunnelNode))
-	{
-		TunnelContext = BlueprintContext.Pin()->GetFunctionContext(CompositeNode->BoundGraph->GetFName());
-	}
-	else if (const UK2Node_MacroInstance* MacroNode = Cast<UK2Node_MacroInstance>(TunnelNode))
+	bool bIsFromThisTunnel = false;
+	UEdGraph* PinGraph = nullptr;
+	const UEdGraphNode* PinNode = TunnelPin ? TunnelPin->GetOwningNode() : nullptr;
+	if (PinNode)
 	{
-		TunnelContext = BlueprintContext.Pin()->GetFunctionContext(MacroNode->GetMacroGraph()->GetFName());
+		// if this is a tunnel instance locate the tunnel graph, otherwise just use the outer graph.
+		PinGraph = PinNode->IsA<UK2Node_Tunnel>() ? GetTunnelGraphFromNode(PinNode) : PinNode->GetTypedOuter<UEdGraph>();
+		if (PinGraph)
+		{
+			bIsFromThisTunnel = GraphName == PinGraph->GetFName();
+		}
 	}
-	return TunnelContext;
+	return bIsFromThisTunnel;
 }
 
-TSharedPtr<FScriptExecutionNode> FBlueprintFunctionContext::GetTunnelEntryPoint(const FName TunnelName) const
+void FBlueprintTunnelInstanceContext::DiscoverExitSites(TSharedPtr<FScriptExecutionNode> MappedNode)
 {
-	TSharedPtr<FScriptExecutionNode> EntryPoint;
-	if (const TSharedPtr<FScriptExecutionNode>* FoundEntryPoint = TunnelEntryPointMap.Find(TunnelName))
+	if (MappedNode->IsTunnelExit())
 	{
-		EntryPoint = *FoundEntryPoint;
+		// Map tunnel exit because we are mapping inside a tunnel discovering exit pins.
+		if (StagingEntryPoint.IsValid())
+		{
+			TSharedPtr<FScriptExecutionTunnelExit> ExitNode = StaticCastSharedPtr<FScriptExecutionTunnelExit>(MappedNode);
+			if (UEdGraphPin* ExternalPin = ExitNode->GetExternalPin())
+			{
+				const int32 ScriptCodeOffset = GetCodeLocationFromPin(ExternalPin);
+				StagingEntryPoint->AddExitSite(ScriptCodeOffset, ExitNode);
+			}
+		}
 	}
-	return EntryPoint;
-}
-
-void FBlueprintFunctionContext::GetTunnelExitPoints(const FName TunnelName, TArray<TSharedPtr<FScriptExecutionNode>>& ExitPointsOut) const
-{
-	ExitPointsOut.Reset();
-	if (const TArray<TSharedPtr<FScriptExecutionNode>>* FoundExitPoints = TunnelExitPointMap.Find(TunnelName))
+	else
 	{
-		ExitPointsOut = *FoundExitPoints;
+		for (auto ChildIter : MappedNode->GetChildNodes())
+		{
+			DiscoverExitSites(ChildIter);
+		}
+		for (auto LinkIter : MappedNode->GetLinkedNodes())
+		{
+			DiscoverExitSites(LinkIter.Value);
+		}
 	}
 }
 
-void FBlueprintFunctionContext::MapTunnelIO(UEdGraph* TunnelGraph)
+TSharedPtr<FScriptExecutionNode> FBlueprintTunnelInstanceContext::MapTunnelBoundary(const UEdGraphPin* TunnelPin)
 {
-	// Map graph entry points
-	check (TunnelGraph);
-	TArray<UK2Node_Tunnel*> GraphTunnels;
-	TunnelGraph->GetNodesOfClass<UK2Node_Tunnel>(GraphTunnels);
-	for (auto Tunnel : GraphTunnels)
+	TSharedPtr<FScriptExecutionNode> TunnelBoundaryNode;
+	// If this isn't an internal pin, this is a boundary to a different tunnel.
+	const bool bInternalTunnel = IsTunnelInternal(TunnelPin->GetOwningNode());
+	const FName TunnelBoundaryName = bInternalTunnel ? GetPinName(TunnelPin) : GetTunnelBoundaryName(TunnelPin);
+	TunnelBoundaryNode = GetProfilerDataForNodeChecked(TunnelBoundaryName);
+	if (TunnelBoundaryNode->IsTunnelExit())
 	{
-		for (auto Pin : Tunnel->Pins)
+		// Map tunnel exit because we are mapping inside a tunnel discovering exit pins.
+		TSharedPtr<FScriptExecutionTunnelExit> ExitNode = StaticCastSharedPtr<FScriptExecutionTunnelExit>(TunnelBoundaryNode);
+		if (UEdGraphPin* ExternalPin = ExitNode->GetExternalPin())
 		{
-			if (Pin->LinkedTo.Num() && Pin->PinType.PinCategory == UEdGraphSchema_K2::PC_Exec)
-			{
-				if (Pin->Direction == EGPD_Output)
-				{
-					// Create entry point
-					FScriptExecNodeParams LinkParams;
-					LinkParams.NodeName = GetPinName(Pin);
-					LinkParams.ObservedObject = Pin->GetOwningNode();
-					LinkParams.ObservedPin = Pin;
-					LinkParams.DisplayName = Pin->GetDisplayName().IsEmpty() ? FText::FromName(LinkParams.NodeName) : Pin->GetDisplayName();
-					LinkParams.Tooltip = LOCTEXT("ExecPin_ExpandTunnelEntryPoint_ToolTip", "Expand tunnel entry point");
-					LinkParams.NodeFlags = EScriptExecutionNodeFlags::TunnelEntryPin;
-					LinkParams.IconColor = FLinearColor(1.f, 1.f, 1.f, 0.8f);
-					const FSlateBrush* Icon = FEditorStyle::GetBrush(TEXT("BlueprintProfiler.BPPinConnected"));
-					LinkParams.Icon = const_cast<FSlateBrush*>(Icon);
-					TSharedPtr<FScriptExecutionNode> EntryPoint = CreateExecutionNode(LinkParams);
-					AddEntryPoint(EntryPoint);
-				}
-				else if (Pin->Direction == EGPD_Input)
-				{
-					// Create exit point
-					FScriptExecNodeParams LinkParams;
-					LinkParams.NodeName = GetPinName(Pin);
-					LinkParams.ObservedObject = Pin->GetOwningNode();
-					LinkParams.ObservedPin = Pin;
-					LinkParams.DisplayName = Pin->GetDisplayName().IsEmpty() ? FText::FromName(LinkParams.NodeName) : Pin->GetDisplayName();
-					LinkParams.Tooltip = LOCTEXT("ExecPin_ExpandTunnelExitPoint_ToolTip", "Expand tunnel exit point");
-					LinkParams.NodeFlags = EScriptExecutionNodeFlags::TunnelExitPin;
-					LinkParams.IconColor = FLinearColor(1.f, 1.f, 1.f, 0.8f);
-					const FSlateBrush* Icon = FEditorStyle::GetBrush(TEXT("BlueprintProfiler.BPPinConnected"));
-					LinkParams.Icon = const_cast<FSlateBrush*>(Icon);
-					TSharedPtr<FScriptExecutionNode> ExitPoint = CreateTypedExecutionNode<FScriptExecutionTunnelExit>(LinkParams);
-					AddExitPoint(ExitPoint);
-				}
-			}
+			const int32 ScriptCodeOffset = GetCodeLocationFromPin(ExternalPin);
+			StagingEntryPoint->AddExitSite(ScriptCodeOffset, ExitNode);
 		}
 	}
-}
-
-int32 FBlueprintFunctionContext::GetTunnelIdFromName(const FName TunnelName) const
-{
-	int32 TunnelId = INDEX_NONE;
-	if (const int32* FoundTunnelId = TunnelIds.Find(TunnelName))
+	else if (TunnelBoundaryNode->IsTunnelEntry())
 	{
-		TunnelId = *FoundTunnelId;
+		return FBlueprintFunctionContext::MapTunnelBoundary(TunnelPin);
 	}
-	return TunnelId;
-}
-
-void FBlueprintFunctionContext::AddEntryPoint(TSharedPtr<FScriptExecutionNode> EntryPoint)
-{
-	EntryPoints.Add(EntryPoint);
-	const int32 TunnelId = TunnelIds.Num();
-	TunnelIds.Add(EntryPoint->GetName()) = TunnelId;
-}
-
-void FBlueprintFunctionContext::AddExitPoint(TSharedPtr<FScriptExecutionNode> ExitPoint)
-{
-	ExitPoints.Add(ExitPoint);
-	const int32 TunnelId = TunnelIds.Num();
-	TunnelIds.Add(ExitPoint->GetName()) = TunnelId;
+	return TunnelBoundaryNode;
 }
 
-TSharedPtr<FScriptExecutionTunnelEntry> FBlueprintFunctionContext::GetTunnelEntrySite(const int32 ScriptCodeOffset)
+void FBlueprintTunnelInstanceContext::MapTunnelIO()
 {
-	TSharedPtr<FScriptExecutionTunnelEntry> Result;
-	if (TSharedPtr<FScriptExecutionTunnelEntry>* SearchResult = TunnelEntrySites.Find(ScriptCodeOffset))
+	if (UK2Node_Tunnel* TunnelInstanceGraphNode = TunnelInstanceNode.Get())
 	{
-		Result = *SearchResult;
+		// Map internal tunnel pins to tunnel instance pins
+		TMap<UEdGraphPin*, UEdGraphPin*> EntryPins;
+		TMap<UEdGraphPin*, UEdGraphPin*> ExitPins;
+		TArray<UK2Node_Tunnel*> GraphTunnels;
+		UEdGraph* TunnelGraph = GetTunnelGraphFromNode(TunnelInstanceGraphNode);
+		TunnelGraph->GetNodesOfClass<UK2Node_Tunnel>(GraphTunnels);
+		for (auto Tunnel : GraphTunnels)
+		{
+			if (IsTunnelInternal(Tunnel))
+			{
+				for (UEdGraphPin* InternalPin : Tunnel->Pins)
+				{
+					if (InternalPin->LinkedTo.Num() && InternalPin->PinType.PinCategory == UEdGraphSchema_K2::PC_Exec)
+					{
+						if (UEdGraphPin* TunnelInstancePin = FindMatchingPin(TunnelInstanceGraphNode, InternalPin))
+						{
+							if (InternalPin->Direction == EGPD_Output)
+							{
+								EntryPins.Add(InternalPin) = TunnelInstancePin;
+							}
+							else
+							{
+								ExitPins.Add(InternalPin) = TunnelInstancePin;
+							}
+						}
+					}
+				}
+			}
+		}
+		// Create instance entry points
+		for (auto EntryPinSet : EntryPins)
+		{
+			FScriptExecNodeParams LinkParams;
+			LinkParams.NodeName = GetPinName(EntryPinSet.Key);
+			LinkParams.ObservedObject = EntryPinSet.Key->GetOwningNode();
+			LinkParams.ObservedPin = EntryPinSet.Key;
+			LinkParams.DisplayName = EntryPinSet.Key->GetDisplayName();
+			LinkParams.Tooltip = LOCTEXT("ExecPin_ExpandTunnelEntryPoint_ToolTip", "Expand tunnel entry point");
+			LinkParams.NodeFlags = EScriptExecutionNodeFlags::TunnelEntryPin|EScriptExecutionNodeFlags::InvalidTrace;
+			LinkParams.IconColor = FLinearColor(1.f, 1.f, 1.f, 0.8f);
+			const bool bPinLinked = EntryPinSet.Key->LinkedTo.Num() > 0;
+			const FSlateBrush* Icon = bPinLinked ?	FEditorStyle::GetBrush(TEXT("BlueprintProfiler.BPPinConnected")) : 
+													FEditorStyle::GetBrush(TEXT("BlueprintProfiler.BPPinDisconnected"));
+			LinkParams.Icon = const_cast<FSlateBrush*>(Icon);
+			TSharedPtr<FScriptExecutionNode> EntryPoint = CreateExecutionNode(LinkParams);
+			AddEntryPoint(EntryPoint);
+		}
+		// Create instance exit points
+		for (auto ExitPinSet : ExitPins)
+		{
+			FScriptExecNodeParams LinkParams;
+			LinkParams.NodeName = GetPinName(ExitPinSet.Key);
+			LinkParams.ObservedObject = ExitPinSet.Key->GetOwningNode();
+			LinkParams.ObservedPin = ExitPinSet.Key;
+			LinkParams.DisplayName = ExitPinSet.Key->GetDisplayName();
+			LinkParams.Tooltip = LOCTEXT("ExecPin_ExpandTunnelExitPoint_ToolTip", "Expand tunnel exit point");
+			LinkParams.NodeFlags = EScriptExecutionNodeFlags::TunnelExitPin;
+			LinkParams.IconColor = FLinearColor(1.f, 1.f, 1.f, 0.8f);
+			const bool bPinLinked = ExitPinSet.Key->LinkedTo.Num() > 0;
+			const FSlateBrush* Icon = bPinLinked ?	FEditorStyle::GetBrush(TEXT("BlueprintProfiler.BPPinConnected")) : 
+													FEditorStyle::GetBrush(TEXT("BlueprintProfiler.BPPinDisconnected"));
+			LinkParams.Icon = const_cast<FSlateBrush*>(Icon);
+			TSharedPtr<FScriptExecutionTunnelExit> ExitPoint = CreateTypedExecutionNode<FScriptExecutionTunnelExit>(LinkParams);
+			const FName TunnelBoundaryName = GetTunnelBoundaryName(ExitPinSet.Value);
+			// Add exit point under internal pin name for mapping.
+			ExternalNodes.Add(TunnelBoundaryName) = ExitPoint;
+			ExitPoint->SetExternalPin(ExitPinSet.Value);
+			AddExitPoint(ExitPoint);
+		}
 	}
-	return Result;
 }
 
-void FBlueprintFunctionContext::GetTunnelsFromExitSite(const int32 ScriptCodeOffset, TArray<TSharedPtr<FScriptExecutionTunnelEntry>>& ResultsOut) const
+TSharedPtr<FScriptExecutionNode> FBlueprintTunnelInstanceContext::CreateExecutionNode(FScriptExecNodeParams& InitParams)
 {
-	TunnelExitSites.MultiFind(ScriptCodeOffset, ResultsOut);
+	// Use the function context that represents the tunnel graph.
+	return TunnelFunctionContext->CreateExecutionNode(InitParams);
 }
 
 //////////////////////////////////////////////////////////////////////////
@@ -1563,7 +1629,7 @@ void FBlueprintFunctionContext::GetTunnelsFromExitSite(const int32 ScriptCodeOff
 bool FScriptEventPlayback::Process(const TArray<FScriptInstrumentedEvent>& SignalData, const int32 StartIdx, const int32 StopIdx)
 {
 	const int32 NumEvents = (StopIdx+1) - StartIdx;
-	ProcessingState = EEventProcessingResult::Failed;
+	bool bProcessingSuccess = false;
 	const bool bEventIsResuming = SignalData[StartIdx].IsResumeEvent() ;
 	IBlueprintProfilerInterface* Profiler = FModuleManager::GetModulePtr<IBlueprintProfilerInterface>("BlueprintProfiler");
 
@@ -1576,9 +1642,10 @@ bool FScriptEventPlayback::Process(const TArray<FScriptInstrumentedEvent>& Signa
 		TSharedPtr<FScriptExecutionNode> EventNode = FunctionContext->GetProfilerDataForNode(EventName);
 		// Find Associated graph nodes and submit into a node map for later processing.
 		TMap<const UEdGraphNode*, NodeSignalHelper> NodeMap;
-		ProcessingState = EEventProcessingResult::Success;
+		bProcessingSuccess = true;
 		int32 LastEventIdx = StartIdx;
 		const int32 EventStartOffset = SignalData[StartIdx].IsResumeEvent() ? 3 : 1;
+		LatentLinkId = bEventIsResuming ? INDEX_NONE : LatentLinkId;
 
 		for (int32 SignalIdx = StartIdx + EventStartOffset; SignalIdx < StopIdx; ++SignalIdx)
 		{
@@ -1616,9 +1683,27 @@ bool FScriptEventPlayback::Process(const TArray<FScriptInstrumentedEvent>& Signa
 				{
 					CurrentNodeData.BlueprintContext = BlueprintContext;
 					CurrentNodeData.FunctionContext = FunctionContext;
-					CurrentNodeData.ImpureNode = FunctionContext->GetProfilerDataForNode(GraphNode->GetFName());
+					CurrentNodeData.ImpureNode = FunctionContext->GetProfilerDataForGraphNode(GraphNode);
 					check(CurrentNodeData.ImpureNode.IsValid());
 				}
+				// Check for tunnel boundries and process here
+				if (CurrentNodeData.ImpureNode->IsTunnelInstance())
+				{
+					TWeakObjectPtr<const UEdGraphNode> NodeToUse = GraphNode;
+					if (const UEdGraphNode* RealSourceNode = FunctionContext->FindTunnelSourceNodeFromCodeLocation(CurrSignal.GetScriptCodeOffset()))
+					{
+						NodeToUse = RealSourceNode;
+					}
+					NodeSignalHelper& TunnelNodeData = NodeMap.FindOrAdd(NodeToUse.Get());
+					if (!TunnelNodeData.FunctionContext.IsValid())
+					{
+						TunnelNodeData.BlueprintContext = BlueprintContext;
+						TunnelNodeData.FunctionContext = FunctionContext;
+					}
+					ProcessTunnelBoundary(TunnelNodeData, CurrSignal);
+					continue;
+				}
+				// Process node data
 				switch (CurrSignal.GetType())
 				{
 					case EScriptInstrumentation::PureNodeEntry:
@@ -1655,7 +1740,6 @@ bool FScriptEventPlayback::Process(const TArray<FScriptInstrumentedEvent>& Signa
 							}
 						}
 						// Process tunnel timings
-						ProcessTunnelEntryPoints(FunctionContext, CurrSignal);
 						CurrentNodeData.ExclusiveTracePaths.Push(TracePath);
 						CurrentNodeData.ExclusiveEvents.Add(CurrSignal);
 						AddToTraceHistory(CurrentNodeData.ImpureNode->GetName(), CurrSignal);
@@ -1693,11 +1777,49 @@ bool FScriptEventPlayback::Process(const TArray<FScriptInstrumentedEvent>& Signa
 						{
 							ProcessExecutionSequence(CurrentNodeData, CurrSignal);
 						}
+						// Process tunnel if present
+						TSharedPtr<FScriptExecutionTunnelEntry> TunnelEntry = TracePath.GetTunnel();
+						if (TunnelEntry.IsValid())
+						{
+							TSharedPtr<FScriptExecutionTunnelInstance> TunnelInstance = TunnelEntry->GetTunnelInstance();
+							check (TunnelInstance.IsValid());
+							const UEdGraphNode* TunnelNode = Cast<UEdGraphNode>(TunnelInstance->GetObservedObject());
+							NodeSignalHelper& TunnelNodeData = NodeMap.FindOrAdd(TunnelNode);
+							FScriptInstrumentedEvent OverrideEvent(CurrSignal);
+							FTracePath TunnelTrace;
+							TracePath.GetTunnelTracePath(TunnelTrace);
+							OverrideEvent.OverrideType(EScriptInstrumentation::NodeEntry);
+							TunnelNodeData.ExclusiveEvents.Add(OverrideEvent);
+							TunnelNodeData.ExclusiveTracePaths.Add(TunnelTrace);
+							TunnelTraceStack.Push(TunnelTrace);
+						}
 						break;
 					}
 					case EScriptInstrumentation::SuspendState:
 					{
-						ProcessingState = EEventProcessingResult::Suspended;
+						// Handle latent suspends - make use of the link id to match re-entry.
+						TSharedPtr<FScriptExecutionInstance> InstanceNode = BlueprintContext->GetInstanceExecNode(InstanceName);
+						if (UObject* InstanceObj = InstanceNode->GetActiveObject())
+						{
+							if (UWorld* WorldForEvent = InstanceObj->GetWorld())
+							{
+								FLatentActionManager& LatentActionManager = WorldForEvent->GetLatentActionManager();
+								TSet<int32> UUIDSet;
+								int32 UUID = INDEX_NONE;
+								LatentActionManager.GetActiveUUIDs(InstanceObj, UUIDSet);
+								for (auto SetEntry : UUIDSet)
+								{
+									UUID = SetEntry;
+								}
+								if (UUID != INDEX_NONE)
+								{
+									if (FDelayAction* Action = LatentActionManager.FindExistingAction<FDelayAction>(InstanceObj, UUID))
+									{
+										LatentLinkId = Action->OutputLink;
+									}
+								}
+							}
+						}
 						break;
 					}
 					case EScriptInstrumentation::NodeExit:
@@ -1711,8 +1833,6 @@ bool FScriptEventPlayback::Process(const TArray<FScriptInstrumentedEvent>& Signa
 								TracePath = CurrentNodeData.ExclusiveTracePaths.Last();
 							}
 						}
-						// Process tunnel timings
-						ProcessTunnelEntryPoints(FunctionContext, CurrSignal);
 						// Add Trace History
 						AddToTraceHistory(CurrentNodeData.ImpureNode->GetName(), CurrSignal);
 						// Process node exit
@@ -1724,7 +1844,7 @@ bool FScriptEventPlayback::Process(const TArray<FScriptInstrumentedEvent>& Signa
 						CurrentNodeData.ExclusiveEvents.Add(CurrSignal);
 						// Process cyclic linkage - reset to root link tracepath
 						TSharedPtr<FScriptExecutionNode> NextLink = CurrentNodeData.ImpureNode->GetLinkedNodeByScriptOffset(ScriptCodeExit);
-						if (NextLink.IsValid() && NextLink->HasFlags(EScriptExecutionNodeFlags::CyclicLinkage))
+						if (NextLink.IsValid() && NextLink->HasCyclicLinkage())
 						{
 							const UEdGraphNode* LinkedGraphNode = Cast<UEdGraphNode>(NextLink->GetObservedObject());
 							if (NodeSignalHelper* LinkedEntry = NodeMap.Find(LinkedGraphNode))
@@ -1756,25 +1876,22 @@ bool FScriptEventPlayback::Process(const TArray<FScriptInstrumentedEvent>& Signa
 			EventTimings.Add(EventName) = SignalData[StopIdx].GetTime() - SignalData[LastEventIdx].GetTime();
 		}
 		// Process outstanding event timings, adding to previous timings if existing.
-		if (ProcessingState != EEventProcessingResult::Suspended)
+		for (auto EventTiming : EventTimings)
 		{
-			for (auto EventTiming : EventTimings)
+			FTracePath EventTracePath;
+			// The UCS, along with BP events declared as 'const' in native C++ code, are implemented as standalone functions, and not within the ubergraph function context.
+			FunctionContext = BlueprintContext->GetFunctionContext(EventTiming.Key);
+			if (!FunctionContext.IsValid())
 			{
-				FTracePath EventTracePath;
-				// The UCS, along with BP events declared as 'const' in native C++ code, are implemented as standalone functions, and not within the ubergraph function context.
-				FunctionContext = BlueprintContext->GetFunctionContext(EventTiming.Key);
-				if (!FunctionContext.IsValid())
-				{
-					FunctionContext = BlueprintContext->GetFunctionContextForEvent(EventTiming.Key);
-				}
-				check (FunctionContext.IsValid());
-				EventNode = FunctionContext->GetProfilerDataForNode(EventTiming.Key);
-				check (EventNode.IsValid());
-				TSharedPtr<FScriptPerfData> PerfData = EventNode->GetOrAddPerfDataByInstanceAndTracePath(InstanceName, EventTracePath);
-				PerfData->AddEventTiming(SignalData[StopIdx].GetTime() - SignalData[LastEventIdx].GetTime());
+				FunctionContext = BlueprintContext->GetFunctionContextForEvent(EventTiming.Key);
 			}
-			EventTimings.Reset();
+			check (FunctionContext.IsValid());
+			EventNode = FunctionContext->GetProfilerDataForNode(EventTiming.Key);
+			check (EventNode.IsValid());
+			TSharedPtr<FScriptPerfData> PerfData = EventNode->GetOrAddPerfDataByInstanceAndTracePath(InstanceName, EventTracePath);
+			PerfData->AddEventTiming(SignalData[StopIdx].GetTime() - SignalData[LastEventIdx].GetTime());
 		}
+		EventTimings.Reset();
 		// Process Node map timings -- this can probably be rolled back into submission during the initial processing and lose this extra iteration.
 		for (auto CurrentNodeData : NodeMap)
 		{
@@ -1897,7 +2014,44 @@ bool FScriptEventPlayback::Process(const TArray<FScriptInstrumentedEvent>& Signa
 			}
 		}
 	}
-	return ProcessingState != EEventProcessingResult::Failed;
+	return bProcessingSuccess;
+}
+
+void FScriptEventPlayback::ProcessTunnelBoundary(NodeSignalHelper& CurrentNodeData, const FScriptInstrumentedEvent& CurrSignal)
+{
+	// Find grab the tunnel instance exec node.
+	TSharedPtr<FScriptExecutionTunnelInstance> TunnelInstance = StaticCastSharedPtr<FScriptExecutionTunnelInstance>(CurrentNodeData.ImpureNode);
+	// Then the Boundary node by script offset.
+	const int32 ScriptCodeOffset = CurrSignal.GetScriptCodeOffset();
+	TSharedPtr<FScriptExecutionNode> TunnelBoundary = TunnelInstance->FindBoundarySite(ScriptCodeOffset);
+	if (TunnelBoundary.IsValid())
+	{
+		if (TunnelBoundary->HasFlags(EScriptExecutionNodeFlags::TunnelEntry))
+		{
+			// Process tunnel entry sites
+			TunnelTraceStack.Push(TracePath);
+			TracePath = FTracePath(TracePath, StaticCastSharedPtr<FScriptExecutionTunnelEntry>(TunnelBoundary));
+			CurrentNodeData.ExclusiveEvents.Add(CurrSignal);
+		}
+		else if (TunnelBoundary->HasFlags(EScriptExecutionNodeFlags::TunnelExit))
+		{
+			// Process tunel exit sites.
+			if (CurrentNodeData.ExclusiveEvents.Num())
+			{
+				TSharedPtr<FScriptExecutionTunnelEntry> Tunnel = TracePath.GetTunnel();
+				FTracePath InternalExitTrace = TracePath;
+				TracePath = TunnelTraceStack.Pop();
+				if (Tunnel.IsValid())
+				{
+					// Update the entry and both exit sites ( one inside the tunnel and one exit site )
+					const double TunnelTiming = CurrSignal.GetTime() - CurrentNodeData.ExclusiveEvents.Last().GetTime();
+					Tunnel->AddTunnelTiming(InstanceName, TracePath, InternalExitTrace, ScriptCodeOffset, TunnelTiming);
+				}
+				TracePath.AddExitPin(ScriptCodeOffset);
+				CurrentNodeData.ExclusiveEvents.Reset();
+			}
+		}
+	}
 }
 
 void FScriptEventPlayback::ProcessExecutionSequence(NodeSignalHelper& CurrentNodeData, const FScriptInstrumentedEvent& CurrSignal)
@@ -1939,44 +2093,6 @@ void FScriptEventPlayback::ProcessExecutionSequence(NodeSignalHelper& CurrentNod
 	}
 }
 
-void FScriptEventPlayback::ProcessTunnelEntryPoints(TSharedPtr<FBlueprintFunctionContext> FunctionContext, const FScriptInstrumentedEvent& CurrSignal)
-{
-	const int32 ScriptOffset = CurrSignal.GetScriptCodeOffset();
-	if (const UEdGraphPin* TunnelPin = FunctionContext->GetPinFromCodeLocation(ScriptOffset))
-	{
-		if (const UK2Node_Tunnel* TunnelInstance = Cast<UK2Node_Tunnel>(TunnelPin->GetOwningNode()))
-		{
-			TArray<TSharedPtr<FScriptExecutionTunnelEntry>> Sites;
-			FunctionContext->GetTunnelsFromExitSite(ScriptOffset, Sites);
-			for (auto ExitSite : Sites)
-			{
-				if (TunnelEventHelper* OpenTunnel = ActiveTunnels.Find(ExitSite->GetName()))
-				{
-					const double TunnelTiming = CurrSignal.GetTime() - OpenTunnel->TunnelEntryTime;
-					TSharedPtr<FScriptExecutionTunnelEntry> TunnelEntryNode = OpenTunnel->TunnelEntryPoint;
-					TunnelEntryNode->UpdateTunnelExit(InstanceName, TracePath, ScriptOffset);
-					if (TunnelEntryNode->IsFinalExitSite(ScriptOffset))
-					{
-						TunnelEntryNode->AddTunnelTiming(InstanceName, OpenTunnel->EntryTracePath, TunnelTiming);
-						TracePath = OpenTunnel->EntryTracePath;
-						ActiveTunnels.Remove(ExitSite->GetName());
-					}
-				}
-			}
-			// Handle tunnel entries
-			TSharedPtr<FScriptExecutionTunnelEntry> EntrySite = FunctionContext->GetTunnelEntrySite(ScriptOffset);
-			if (EntrySite.IsValid())
-			{
-				TunnelEventHelper& NewTunnel = ActiveTunnels.Add(EntrySite->GetName());
-				NewTunnel.EntryTracePath = TracePath;
-				NewTunnel.EntryTracePath.AddExitPin(ScriptOffset);
-				NewTunnel.TunnelEntryPoint = EntrySite;
-				NewTunnel.TunnelEntryTime = CurrSignal.GetTime();
-			}
-		}
-	}
-}
-
 void FScriptEventPlayback::AddToTraceHistory(const FName NodeName, const FScriptInstrumentedEvent& TraceSignal)
 {
 	// Add Trace History
diff --git a/Engine/Source/Developer/BlueprintProfiler/Public/BlueprintProfiler.h b/Engine/Source/Developer/BlueprintProfiler/Public/BlueprintProfiler.h
--- a/Engine/Source/Developer/BlueprintProfiler/Public/BlueprintProfiler.h
+++ b/Engine/Source/Developer/BlueprintProfiler/Public/BlueprintProfiler.h
@@ -88,8 +88,8 @@ protected:
 	EBlueprintProfilerHeatMapDisplayMode::Type GraphNodeHeatMapDisplayMode;
 	/** Current wire heat map display mode */
 	EBlueprintProfilerHeatMapDisplayMode::Type WireHeatMapDisplayMode;
-	/** Suspended script events */
-	TMap<FName, TSharedPtr<FScriptEventPlayback>> SuspendedEvents;
+	/** Suspended script events, the key is the latent LinkId */
+	TMap<FName, TMap<int32, TSharedPtr<FScriptEventPlayback>>> SuspendedEvents;
 	/** Cached object path and code offset lookup to UObjects */
 	TMap<FString, TSharedPtr<FBlueprintExecutionContext>> PathToBlueprintContext;
 	/** The profiler connection factory */
diff --git a/Engine/Source/Developer/BlueprintProfiler/Public/ScriptInstrumentationPlayback.h b/Engine/Source/Developer/BlueprintProfiler/Public/ScriptInstrumentationPlayback.h
--- a/Engine/Source/Developer/BlueprintProfiler/Public/ScriptInstrumentationPlayback.h
+++ b/Engine/Source/Developer/BlueprintProfiler/Public/ScriptInstrumentationPlayback.h
@@ -73,8 +73,11 @@ public:
 	/** Maps the instance returning the system name for instance */
 	BLUEPRINTPROFILER_API FName MapBlueprintInstance(const FString& InstancePath);
 
+	/** Returns the instance if mapped */
+	BLUEPRINTPROFILER_API TSharedPtr<class FScriptExecutionInstance> GetInstanceExecNode(const FName InstanceName);
+
 	/** Remaps PIE actor instance paths to editor actor instances */
-	BLUEPRINTPROFILER_API FName RemapInstancePath(const FName InstanceName) const ;
+	BLUEPRINTPROFILER_API FName RemapInstancePath(const FName InstanceName) const;
 
 	/** Returns the function name containing the event */
 	FName GetEventFunctionName(const FName EventName) const;
@@ -109,7 +112,7 @@ public:
 private:
 
 	/** Creates and returns a new function context, returns existing if present  */
-	TSharedPtr<FBlueprintFunctionContext> CreateFunctionContext(const FName FunctionName, UEdGraph* Graph);
+	template<typename FunctionType> TSharedPtr<FunctionType> CreateFunctionContext(const FName FunctionName, UEdGraph* Graph);
 
 	/** Walks the blueprint and maps events and functions ready for profiling data */
 	bool MapBlueprintExecution();
@@ -150,11 +153,16 @@ class FBlueprintFunctionContext : public TSharedFromThis<FBlueprintFunctionConte
 {
 public:
 
+	virtual ~FBlueprintFunctionContext() {}
+
 	/** Initialise the function context from the graph */
-	void InitialiseContextFromGraph(TSharedPtr<FBlueprintExecutionContext> BlueprintContext, const FName FunctionNameIn, UEdGraph* Graph);
+	virtual void InitialiseContextFromGraph(TSharedPtr<FBlueprintExecutionContext> BlueprintContext, const FName FunctionNameIn, UEdGraph* Graph);
+
+	/** Discover and add any tunnels */
+	void DiscoverTunnels(UEdGraph* Graph, TMap<UK2Node_Tunnel*, TSharedPtr<FBlueprintFunctionContext>>& DiscoveredTunnels);
 
 	/** Map Function */
-	void MapFunction();
+	virtual void MapFunction();
 
 	/** Returns if the function context is fully formed */
 	bool IsContextValid() const { return Function.IsValid() && BlueprintClass.IsValid(); }
@@ -162,6 +170,9 @@ public:
 	/** Returns the graph node at the specified script offset */
 	const UEdGraphNode* GetNodeFromCodeLocation(const int32 ScriptOffset);
 
+	/** Returns the tunnel source node at the specified script offset */
+	const UEdGraphNode* FindTunnelSourceNodeFromCodeLocation(const int32 ScriptOffset);
+
 	/** Returns the pin at the specified script offset */
 	const UEdGraphPin* GetPinFromCodeLocation(const int32 ScriptOffset);
 
@@ -180,30 +191,55 @@ public:
 	/** Returns the execution node representing the node name */
 	TSharedPtr<FScriptExecutionNode> GetProfilerDataForNode(const FName NodeName);
 
+	/** Returns the execution node representing the node name, asserting on fail */
+	TSharedPtr<FScriptExecutionNode> GetProfilerDataForNodeChecked(const FName NodeName);
+
+	/** Returns the execution node representing the graph node */
+	TSharedPtr<FScriptExecutionNode> GetProfilerDataForGraphNode(const UEdGraphNode* Node);
+
 	/** Returns the execution node representing the node name */
 	template<typename ExecNodeType> TSharedPtr<ExecNodeType> GetTypedProfilerDataForNode(const FName NodeName);
 
+	/** Returns the execution node representing the object */
+	template<typename ExecNodeType> TSharedPtr<ExecNodeType> GetTypedProfilerDataForGraphNode(const UEdGraphNode* Node);
+
 	/** Adds all function entry points to node as children */
 	void AddCallSiteEntryPointsToNode(TSharedPtr<FScriptExecutionNode> CallingNode) const;
 
 	/** Returns the function name for this context */
 	FName GetFunctionName() const { return FunctionName; }
 
-	/** Returns the tunnel type */
-	const FScriptExecutionTunnelEntry::ETunnelType GetTunnelType(const FName TunnelName) const;
+	/** Returns the UFunction for this context */
+	TWeakObjectPtr<UFunction> GetUFunction() { return Function; }
 
-	/** Returns the tunnel entry from the script code offset */
-	TSharedPtr<FScriptExecutionTunnelEntry> GetTunnelEntrySite(const int32 ScriptCodeOffset);
+	/** Returns the BlueprintClass for this context */
+	TWeakObjectPtr<UBlueprintGeneratedClass> GetBlueprintClass() { return BlueprintClass; }
 
-	/** Returns the tunnel custom entry from the script code offset */
-	void GetTunnelsFromExitSite(const int32 ScriptCodeOffset, TArray<TSharedPtr<FScriptExecutionTunnelEntry>>& ResultsOut) const;
+	/** Returns a valid pin name, creating one based on pin characteristics if not available */
+	static FName GetPinName(const UEdGraphPin* Pin);
 
-private:
+	/** Returns a fully qualified pin name including owning node id */
+	static FName GetUniquePinName(const UEdGraphPin* Pin);
+
+	/** Returns a tunnel instance Boundary name */
+	static FName GetTunnelBoundaryName(const UEdGraphPin* Pin);
+
+	/** Looks for matching pin in the supplied node, handles pins with no name. */
+	static UEdGraphPin* FindMatchingPin(const UEdGraphNode* NodeToSearch, const UEdGraphPin* PinToFind, const bool bIgnoreDirection = false);
+
+	/** Utility function that returns the tunnel graph from the supplied node */
+	static UEdGraph* GetTunnelGraphFromNode(const UEdGraphNode* TunnelNode);
+
+	/** Utility function that returns if the tunnel node is internal to a tunnel graph */
+	static bool IsTunnelInternal(const UEdGraphNode* TunnelNode);
+
+protected:
 
 	friend FBlueprintExecutionContext;
+	friend class FBlueprintTunnelInstanceContext;
 
 	/** Utility to create execution node */
-	TSharedPtr<FScriptExecutionNode> CreateExecutionNode(FScriptExecNodeParams& InitParams);
+	virtual TSharedPtr<FScriptExecutionNode> CreateExecutionNode(FScriptExecNodeParams& InitParams);
 
 	/** Utility to create typed execution node */
 	template<typename ExecNodeType> TSharedPtr<ExecNodeType> CreateTypedExecutionNode(FScriptExecNodeParams& InitParams);
@@ -214,11 +250,8 @@ private:
 	/** Adds a new function exit point */
 	void AddExitPoint(TSharedPtr<FScriptExecutionNode> ExitPoint);
 
-	/** Locates and generates event contexts for any input events passed in */
-	void CreateInputEvents(TSharedPtr<FBlueprintExecutionContext> BlueprintContextIn, const TArray<UK2Node*>& InputEventNodes);
-
-	/** Locates and generates event contexts for any async task events passed in */
-	void CreateAsyncTaskEvents(TSharedPtr<FBlueprintExecutionContext> BlueprintContextIn, const TArray<UK2Node_BaseAsyncTask*>& AsyncTaskNodes);
+	/** Creates an event for delegate pin entry points */
+	void CreateDelegatePinEvents(TSharedPtr<FBlueprintExecutionContext> BlueprintContextIn, const TMap<FName, FEdGraphPinReference>& PinEvents);
 
 	/** Processes and detects any cyclic links making the linkage safe for traversal */
 	bool DetectCyclicLinks(TSharedPtr<FScriptExecutionNode> ExecNode, TSet<TSharedPtr<FScriptExecutionNode>>& Filter);
@@ -226,7 +259,7 @@ private:
 	// --Execution mapping functionality
 
 	/** Maps each blueprint node following execution wires */
-	TSharedPtr<FScriptExecutionNode> MapNodeExecution(UEdGraphNode* NodeToMap);
+	virtual TSharedPtr<FScriptExecutionNode> MapNodeExecution(UEdGraphNode* NodeToMap);
 
 	/** Maps input pin execution */
 	void MapInputPins(TSharedPtr<FScriptExecutionNode> ExecNode, const TArray<UEdGraphPin*>& Pins);
@@ -234,36 +267,18 @@ private:
 	/** Maps pin execution */
 	void MapExecPins(TSharedPtr<FScriptExecutionNode> ExecNode, const TArray<UEdGraphPin*>& Pins);
 
-	/** Returns a valid pin name, creating one based on pin characteristics if not available */
-	FName GetPinName(const UEdGraphPin* Pin) const;
-
-	/** Returns a full qualified pin name including owning node id */
-	FName GetUniquePinName(const UEdGraphPin* Pin) const;
-
-	/** Maps execution pin tunnel entry */
-	TSharedPtr<FScriptExecutionNode> MapTunnelEntry(const FName TunnelName, const UEdGraphPin* ExecPin, const UEdGraphPin* TunnelPin, const int32 PinScriptOffset);
-
 	// --Tunnel mapping functionality
 
+	/** Maps tunnel boundries */
+	virtual TSharedPtr<FScriptExecutionNode> MapTunnelBoundary(const UEdGraphPin* TunnelPin);
+
 	/** Maps the tunnel point into the instanced graph, creating the instanced graph if not already existing. */
 	void MapTunnelInstance(UK2Node_Tunnel* TunnelInstance);
 
-	/** Returns the tunnel function context that the tunnel node invokes */
-	TSharedPtr<FBlueprintFunctionContext> GetTunnelContextFromNode(const UEdGraphNode* TunnelNode) const;
-
-	/** Returns the registered tunnel execution entry point given the tunnel name */
-	TSharedPtr<FScriptExecutionNode> GetTunnelEntryPoint(const FName TunnelName) const;
-
-	/** Returns the registered tunnel exit points given the tunnel name */
-	void GetTunnelExitPoints(const FName TunnelName, TArray<TSharedPtr<FScriptExecutionNode>>& ExitPointsOut) const;
-
-	/** Maps input and output pins in a tunnel graph */
-	void MapTunnelIO(UEdGraph* TunnelGraph);
+	/** Maps execution pin tunnel entry */
+	void MapTunnelExits(TSharedPtr<FScriptExecutionTunnelEntry> TunnelEntryPoint);
 
-	/** Returns the tunnel Id associated with the tunnel name */
-	int32 GetTunnelIdFromName(const FName TunnelName) const;
-
-private:
+protected:
 
 	/** The graph name */
 	FName GraphName;
@@ -283,38 +298,64 @@ private:
 	TMap<int32, FEdGraphPinReference> ScriptOffsetToPins;
 	/** Execution nodes containing profiling data */
 	TMap<FName, TSharedPtr<FScriptExecutionNode>> ExecutionNodes;
+	/** Other function contexts in use by this context */
+	TMap<FName, TWeakPtr<FBlueprintFunctionContext>> ChildFunctionContexts;
 	/** Graph entry points */
 	TArray<TSharedPtr<FScriptExecutionNode>> EntryPoints;
 	/** Graph exit points */
 	TArray<TSharedPtr<FScriptExecutionNode>> ExitPoints;
-	/** Map associating function entry points to execution points */
-	TMap<FName, TSharedPtr<FScriptExecutionNode>> TunnelEntryPointMap;
-	/** Map associating function entry points to exit points */
-	TMap<FName, TArray<TSharedPtr<FScriptExecutionNode>>> TunnelExitPointMap;
-	/** Map associating tunnel entry points to script code offsets */
-	TMultiMap<int32, TSharedPtr<FScriptExecutionTunnelEntry>> TunnelEntrySites;
-	/** Map associating custom tunnel entries to script code offsets */
-	TMultiMap<int32, TSharedPtr<FScriptExecutionTunnelEntry>> TunnelExitSites;
-	/** Staging tunnel name during mapping */
-	FName StagingTunnelName;
-	/** Map of tunnel Ids */
-	TMap<FName, int32> TunnelIds;
 };
 
 //////////////////////////////////////////////////////////////////////////
-// FScriptEventPlayback
+// FBlueprintTunnelInstanceContext
 
-namespace EEventProcessingResult
+class FBlueprintTunnelInstanceContext : public FBlueprintFunctionContext
 {
-	enum Type
-	{
-		None = 0,
-		Failed,
-		Success,
-		Suspended
-	};
+public:
+
+	// ~FBlueprintFunctionContext Begin
+	virtual void InitialiseContextFromGraph(TSharedPtr<FBlueprintExecutionContext> BlueprintContext, const FName TunnelInstanceName, UEdGraph* TunnelGraph);
+	virtual void MapFunction() override {}
+	virtual TSharedPtr<FScriptExecutionNode> CreateExecutionNode(FScriptExecNodeParams& InitParams) override;
+	// ~FBlueprintFunctionContext End
+
+	/** Maps a tunnel function context. */
+	void MapTunnelContext(TSharedPtr<FBlueprintFunctionContext> CallingFunctionContext, UK2Node_Tunnel* TunnelInstance);
+
+private:
+
+	friend FBlueprintExecutionContext;
+
+	// ~FBlueprintFunctionContext Begin
+	virtual TSharedPtr<FScriptExecutionNode> MapNodeExecution(UEdGraphNode* NodeToMap) override;
+	virtual TSharedPtr<FScriptExecutionNode> MapTunnelBoundary(const UEdGraphPin* TunnelPin) override;
+	// ~FBlueprintFunctionContext End
+	
+	/** Maps input and output pins in a tunnel graph */
+	void MapTunnelIO();
+
+	/** Returns true if the pin is part of this tunnel */
+	bool IsPinFromThisTunnel(const UEdGraphPin* TunnelPin) const;
+
+	/** Discovers already mapped exit sites */
+	void DiscoverExitSites(TSharedPtr<FScriptExecutionNode> MappedNode);
+
+private:
+
+	/** The tunnel instance this context represents */
+	TWeakObjectPtr<UK2Node_Tunnel> TunnelInstanceNode;
+	/** The non-instance context of the tunnel graph */
+	TSharedPtr<FBlueprintFunctionContext> TunnelFunctionContext;
+	/** External tunnel nodes */
+	TMap<FName, TSharedPtr<FScriptExecutionNode>> ExternalNodes;
+	/** Staging entry point */
+	TSharedPtr<FScriptExecutionTunnelEntry> StagingEntryPoint;
+
 };
 
+//////////////////////////////////////////////////////////////////////////
+// FScriptEventPlayback
+
 class FScriptEventPlayback : public TSharedFromThis<FScriptEventPlayback>
 {
 public:
@@ -353,16 +394,16 @@ public:
 	FScriptEventPlayback(TSharedPtr<FBlueprintExecutionContext> BlueprintContextIn, const FName InstanceNameIn)
 		: InstanceName(InstanceNameIn)
 		, CurrentFunctionName(NAME_None)
-		, ProcessingState(EEventProcessingResult::None)
+		, LatentLinkId(INDEX_NONE)
 		, BlueprintContext(BlueprintContextIn)
 	{
 	}
 
-	/** Returns the current processing state */
-	EEventProcessingResult::Type GetProcessingState() const { return ProcessingState; }
-
 	/** Returns if current processing state is suspended */
-	bool IsSuspended() const { return ProcessingState == EEventProcessingResult::Suspended; }
+	bool IsSuspended() const { return LatentLinkId != INDEX_NONE; }
+
+	/** Returns the current latent action UUID */
+	const int32 GetLatentLinkId() const { return LatentLinkId; }
 
 	/** Processes the event and cleans up any errant signals */
 	bool Process(const TArray<class FScriptInstrumentedEvent>& Events, const int32 Start, const int32 Stop);
@@ -372,12 +413,12 @@ public:
 
 private:
 
+	/** Process tunel boundries */
+	void ProcessTunnelBoundary(NodeSignalHelper& CurrentNodeData, const FScriptInstrumentedEvent& CurrSignal);
+
 	/** Process execution sequence */
 	void ProcessExecutionSequence(NodeSignalHelper& CurrentNodeData, const FScriptInstrumentedEvent& CurrSignal);
 
-	/** Process tunnel entry points */
-	void ProcessTunnelEntryPoints(TSharedPtr<FBlueprintFunctionContext> FunctionContext, const FScriptInstrumentedEvent& CurrSignal);
-
 	/** Add to trace history */
 	void AddToTraceHistory(const FName NodeName, const FScriptInstrumentedEvent& TraceSignal);
 
@@ -389,12 +430,14 @@ private:
 	FName CurrentFunctionName;
 	/** Event name */
 	FName EventName;
-	/** Processing state */
-	EEventProcessingResult::Type ProcessingState;
+	/** Latent Link Id for suspended Event */
+	int32 LatentLinkId;
 	/** Current tracepath state */
 	FTracePath TracePath;
 	/** Current tracepath stack */
 	TArray<FTracePath> TraceStack;
+	/** Current tunnel tracepath stack */
+	TArray<FTracePath> TunnelTraceStack;
 	/** Active tunnels */
 	TMap<FName, TunnelEventHelper> ActiveTunnels;
 	/** Event Timings */
