Copying //UE4/Dev-Blueprints to //UE4/Dev-Main (Source: //UE4/Dev-Blueprints @ 3057841)

diff --git a/Engine/Source/Developer/BlueprintProfiler/Private/ScriptInstrumentationPlayback.cpp b/Engine/Source/Developer/BlueprintProfiler/Private/ScriptInstrumentationPlayback.cpp
--- a/Engine/Source/Developer/BlueprintProfiler/Private/ScriptInstrumentationPlayback.cpp
+++ b/Engine/Source/Developer/BlueprintProfiler/Private/ScriptInstrumentationPlayback.cpp
@@ -24,13 +24,13 @@ bool FBlueprintExecutionContext::InitialiseContext(const FString& BlueprintPath)
 		{
 			BlueprintClass = BPClass;
 			Blueprint = Cast<UBlueprint>(BPClass->ClassGeneratedBy);
-			UbergraphFunctionName = BPClass->UberGraphFunction ? BPClass->UberGraphFunction->GetFName() : NAME_None;
 		}
 	}
-	if (Blueprint.IsValid() && BlueprintClass.IsValid() && UbergraphFunctionName != NAME_None)
+	if (Blueprint.IsValid() && BlueprintClass.IsValid())
 	{
 		// Create new blueprint exec node
 		FScriptExecNodeParams BlueprintParams;
+		BlueprintParams.SampleFrequency = 1;
 		BlueprintParams.NodeName = FName(*BlueprintPath);
 		BlueprintParams.ObservedObject = Blueprint.Get();
 		BlueprintParams.OwningGraphName = NAME_None;
@@ -52,7 +52,6 @@ void FBlueprintExecutionContext::RemoveMapping()
 	BlueprintClass.Reset();
 	BlueprintNode.Reset();
 	FunctionContexts.Reset();
-	UbergraphFunctionName = NAME_None;
 }
 
 bool FBlueprintExecutionContext::IsEventMapped(const FName EventName) const
@@ -84,6 +83,7 @@ FName FBlueprintExecutionContext::MapBlueprintInstance(const FString& InstancePa
 		{
 			// Create new instance node
 			FScriptExecNodeParams InstanceNodeParams;
+			InstanceNodeParams.SampleFrequency = 1;
 			InstanceNodeParams.NodeName = InstanceName;
 			InstanceNodeParams.ObservedObject = Instance.Get();
 			InstanceNodeParams.NodeFlags = EScriptExecutionNodeFlags::Instance;
@@ -170,6 +170,8 @@ bool FBlueprintExecutionContext::ResolveInstance(FName& InstanceNameInOut, TWeak
 		// Attempt to locate the instance and map PIE objects to editor world objects
 		if (const UObject* ObjectPtr = FindObject<UObject>(nullptr, *InstanceNameInOut.ToString()))
 		{
+			const bool bTransientActor = ObjectPtr->HasAnyFlags(RF_Transient);
+			bool bRuntimeOnlyActor = true;
 			// Get Outer world
 			if (UWorld* ObjectWorld = ObjectPtr->GetTypedOuter<UWorld>())
 			{
@@ -190,13 +192,15 @@ bool FBlueprintExecutionContext::ResolveInstance(FName& InstanceNameInOut, TWeak
 									PIEInstanceNameMap.Add(InstanceNameInOut) = CorrectedName;
 									InstanceNameInOut = CorrectedName;
 									ObjectInOut = EditorObject;
+									bRuntimeOnlyActor = false;
 								}
 								break;
 							}
 						}
 					}
+
 				}
-				else if (!ObjectPtr->HasAnyFlags(RF_Transient) && ObjectPtr->GetClass() == BlueprintClass)
+				if ((bRuntimeOnlyActor || !bTransientActor) && ObjectPtr->GetClass() == BlueprintClass)
 				{
 					EditorActorInstances.Add(InstanceNameInOut) = ObjectPtr;
 					PIEInstanceNameMap.Add(InstanceNameInOut) = InstanceNameInOut;
@@ -209,29 +213,41 @@ bool FBlueprintExecutionContext::ResolveInstance(FName& InstanceNameInOut, TWeak
 	return bNewInstance;
 }
 
-FName FBlueprintExecutionContext::GetEventFunctionName(const FName EventName) const
+TSharedPtr<FBlueprintFunctionContext> FBlueprintExecutionContext::GetFunctionContextForEventChecked(const FName ScopedEventName) const
 {
-	return (EventName == UEdGraphSchema_K2::FN_UserConstructionScript) ? UEdGraphSchema_K2::FN_UserConstructionScript : UbergraphFunctionName;
+	TSharedPtr<FBlueprintFunctionContext> Result;
+	if (const TSharedPtr<FBlueprintFunctionContext>* SearchResult = EventFunctionContexts.Find(ScopedEventName))
+	{
+		Result = *SearchResult;
+	}
+	else
+	{
+		Result = GetFunctionContext(ScopedEventName);
+	}
+	check (Result.IsValid());
+	return Result;
 }
 
-TSharedPtr<FBlueprintFunctionContext> FBlueprintExecutionContext::GetFunctionContextForEvent(const FName EventName) const
+TSharedPtr<FBlueprintFunctionContext> FBlueprintExecutionContext::GetFunctionContext(const FName ScopedFunctionName) const
 {
 	TSharedPtr<FBlueprintFunctionContext> Result;
-	if (const TSharedPtr<FBlueprintFunctionContext>* SearchResult = EventFunctionContexts.Find(EventName))
+	if (const TSharedPtr<FBlueprintFunctionContext>* SearchResult = FunctionContexts.Find(ScopedFunctionName))
 	{
 		Result = *SearchResult;
 	}
 	return Result;
 }
 
-TSharedPtr<FBlueprintFunctionContext> FBlueprintExecutionContext::GetFunctionContext(const FName FunctionNameIn) const
+TSharedPtr<FBlueprintFunctionContext> FBlueprintExecutionContext::GetFunctionContextFromGraph(const UEdGraph* Graph) const
 {
 	TSharedPtr<FBlueprintFunctionContext> Result;
-	FName FunctionName = (FunctionNameIn == UEdGraphSchema_K2::GN_EventGraph) ? UbergraphFunctionName : FunctionNameIn;
-
-	if (const TSharedPtr<FBlueprintFunctionContext>* SearchResult = FunctionContexts.Find(FunctionName))
+	if (Graph)
 	{
-		Result = *SearchResult;
+		const FName ScopedFunctionName = GetScopedFunctionNameFromGraph(Graph);
+		if (const TSharedPtr<FBlueprintFunctionContext>* SearchResult = FunctionContexts.Find(ScopedFunctionName))
+		{
+			Result = *SearchResult;
+		}
 	}
 	return Result;
 }
@@ -251,10 +267,10 @@ template<typename FunctionType> TSharedPtr<FunctionType> FBlueprintExecutionCont
 bool FBlueprintExecutionContext::HasProfilerDataForPin(const UEdGraphPin* GraphPin) const
 {
 	SCOPE_CYCLE_COUNTER(STAT_NodeLookupCost);
-	const UEdGraph* OuterGraph = GraphPin ? GraphPin->GetOwningNode()->GetTypedOuter<UEdGraph>() : nullptr;
-	if (OuterGraph)
+	const UEdGraph* OwningGraph = GraphPin ? FBlueprintFunctionContext::GetGraphFromNode(GraphPin->GetOwningNode()) : nullptr;
+	if (OwningGraph)
 	{
-		TSharedPtr<FBlueprintFunctionContext> FunctionContext = GetFunctionContext(OuterGraph->GetFName());
+		TSharedPtr<FBlueprintFunctionContext> FunctionContext = GetFunctionContextFromGraph(OwningGraph);
 		if (FunctionContext.IsValid())
 		{
 			const FName PinName(FunctionContext->GetUniquePinName(GraphPin));
@@ -268,10 +284,10 @@ TSharedPtr<FScriptExecutionNode> FBlueprintExecutionContext::GetProfilerDataForP
 {
 	SCOPE_CYCLE_COUNTER(STAT_NodeLookupCost);
 	TSharedPtr<FScriptExecutionNode> Result;
-	const UEdGraph* OuterGraph = GraphPin ? GraphPin->GetOwningNode()->GetTypedOuter<UEdGraph>() : nullptr;
-	if (OuterGraph)
+	const UEdGraph* OwningGraph = GraphPin ? FBlueprintFunctionContext::GetGraphFromNode(GraphPin->GetOwningNode()) : nullptr;
+	if (OwningGraph)
 	{
-		TSharedPtr<FBlueprintFunctionContext> FunctionContext = GetFunctionContext(OuterGraph->GetFName());
+		TSharedPtr<FBlueprintFunctionContext> FunctionContext = GetFunctionContextFromGraph(OwningGraph);
 		if (FunctionContext.IsValid())
 		{
 			const FName PinName(FunctionContext->GetUniquePinName(GraphPin));
@@ -296,10 +312,11 @@ TSharedPtr<FScriptExecutionNode> FBlueprintExecutionContext::GetProfilerDataForP
 bool FBlueprintExecutionContext::HasProfilerDataForNode(const UEdGraphNode* GraphNode) const
 {
 	SCOPE_CYCLE_COUNTER(STAT_NodeLookupCost);
-	const UEdGraph* OuterGraph = GraphNode ? GraphNode->GetTypedOuter<UEdGraph>() : nullptr;
-	if (OuterGraph)
+	// Do a simple outer lookup to get the nodes graph, the function context will handle tunnel instances.
+	const UEdGraph* OwningGraph = GraphNode ? GraphNode->GetTypedOuter<UEdGraph>() : nullptr;
+	if (OwningGraph)
 	{
-		TSharedPtr<FBlueprintFunctionContext> FunctionContext = GetFunctionContext(OuterGraph->GetFName());
+		TSharedPtr<FBlueprintFunctionContext> FunctionContext = GetFunctionContextFromGraph(OwningGraph);
 		if (FunctionContext.IsValid())
 		{
 			return FunctionContext->HasProfilerDataForNode(GraphNode->GetFName());
@@ -312,10 +329,11 @@ TSharedPtr<FScriptExecutionNode> FBlueprintExecutionContext::GetProfilerDataForN
 {
 	SCOPE_CYCLE_COUNTER(STAT_NodeLookupCost);
 	TSharedPtr<FScriptExecutionNode> Result;
-	const UEdGraph* OuterGraph = GraphNode ? GraphNode->GetTypedOuter<UEdGraph>() : nullptr;
-	if (OuterGraph)
+	// Do a simple outer lookup to get the nodes graph, the function context will handle tunnel instances.
+	const UEdGraph* OwningGraph = GraphNode ? GraphNode->GetTypedOuter<UEdGraph>() : nullptr;
+	if (OwningGraph)
 	{
-		TSharedPtr<FBlueprintFunctionContext> FunctionContext = GetFunctionContext(OuterGraph->GetFName());
+		TSharedPtr<FBlueprintFunctionContext> FunctionContext = GetFunctionContextFromGraph(OwningGraph);
 		if (FunctionContext.IsValid() && FunctionContext->HasProfilerDataForNode(GraphNode->GetFName()))
 		{
 			Result = FunctionContext->GetProfilerDataForNode(GraphNode->GetFName());
@@ -351,18 +369,17 @@ bool FBlueprintExecutionContext::MapBlueprintExecution()
 		{
 			if (UBlueprintGeneratedClass* CurrBPGC = Cast<UBlueprintGeneratedClass>(CurrBlueprint->GeneratedClass))
 			{
-				const FName CurrUbergraphName = Cast<UBlueprintGeneratedClass>(CurrBlueprint->GeneratedClass)->UberGraphFunction->GetFName();
 				TArray<UEdGraph*> Graphs;
 				CurrBlueprint->GetAllGraphs(Graphs);
 				for (auto Graph : Graphs)
 				{
-					FName FunctionName = Graph->GetFName();
-					FunctionName = FunctionName == UEdGraphSchema_K2::GN_EventGraph ? CurrUbergraphName : FunctionName;
-					if (!FunctionGraphs.Contains(FunctionName))
+					const FName FunctionName = GetFunctionNameFromGraph(Graph);
+					if (UFunction* ScriptFunction = CurrBPGC->FindFunctionByName(FunctionName))
 					{
-						if (UFunction* ScriptFunction = CurrBPGC->FindFunctionByName(FunctionName))
+						const FName ScopedFunctionName = GetScopedFunctionNameFromGraph(Graph);
+						if (!FunctionGraphs.Contains(ScopedFunctionName))
 						{
-							FunctionGraphs.Add(FunctionName) = Graph;
+							FunctionGraphs.Add(ScopedFunctionName) = Graph;
 						}
 					}
 				}
@@ -379,17 +396,12 @@ bool FBlueprintExecutionContext::MapBlueprintExecution()
 		TSet<UEdGraph*> MappedTunnelGraphs;
 		for (auto TunnelInstance : DiscoveredTunnels)
 		{
-			if (UEdGraph* TunnelGraph = TunnelInstance.Value->GetTunnelGraphFromNode(TunnelInstance.Key))
+			if (UEdGraph* TunnelGraph = FBlueprintFunctionContext::GetGraphFromNode(TunnelInstance.Key, false))
 			{
-				if (!MappedTunnelGraphs.Contains(TunnelGraph))
-				{
-					// Create the function context for the tunnel graph.
-					CreateFunctionContext<FBlueprintFunctionContext>(TunnelGraph->GetFName(), TunnelGraph);
-					MappedTunnelGraphs.Add(TunnelGraph);
-				}
 				// Create the tunnel instance context.
-				TSharedPtr<FBlueprintTunnelInstanceContext> NewTunnelContext = CreateFunctionContext<FBlueprintTunnelInstanceContext>(TunnelInstance.Key->GetFName(), TunnelGraph);
-				NewTunnelContext->MapTunnelContext(TunnelInstance.Value, TunnelInstance.Key);
+				const FName TunnelInstanceFunctionName = FBlueprintFunctionContext::GetTunnelInstanceFunctionName(TunnelInstance.Key);
+				TSharedPtr<FBlueprintTunnelInstanceContext> NewTunnelInstanceContext = CreateFunctionContext<FBlueprintTunnelInstanceContext>(TunnelInstanceFunctionName, TunnelGraph);
+				NewTunnelInstanceContext->MapTunnelContext(TunnelInstance.Value, TunnelInstance.Key);
 				TunnelInstance.Value->MapTunnelInstance(TunnelInstance.Key);
 			}
 		}
@@ -398,11 +410,51 @@ bool FBlueprintExecutionContext::MapBlueprintExecution()
 		{
 			Context.Value->MapFunction();
 		}
+		// Sort the events
+		if (BlueprintNode.IsValid())
+		{
+			BlueprintNode->SortEvents();
+		}
 		bMappingSuccessful = true;
 	}
 	return bMappingSuccessful;
 }
 
+FName FBlueprintExecutionContext::GetFunctionNameFromGraph(const UEdGraph* Graph) const
+{
+	FName FunctionName = NAME_None;
+	if (Graph)
+	{
+		UBlueprint* OwnerBlueprint = Graph->GetTypedOuter<UBlueprint>();
+		UBlueprintGeneratedClass* BPGC = OwnerBlueprint ? Cast<UBlueprintGeneratedClass>(OwnerBlueprint->GeneratedClass) : nullptr;
+		if (BPGC)
+		{
+			const bool bIsEventGraph = BPGC->UberGraphFunction && FBlueprintEditorUtils::IsEventGraph(Graph);
+			FunctionName = bIsEventGraph ? BPGC->UberGraphFunction->GetFName() : Graph->GetFName();
+		}
+	}
+	check (FunctionName != NAME_None);
+	return FunctionName;
+}
+
+FName FBlueprintExecutionContext::GetScopedFunctionNameFromGraph(const UEdGraph* Graph) const
+{
+	FName ScopedFunctionName = NAME_None;
+	if (Graph)
+	{
+		UBlueprint* OwnerBlueprint = Graph->GetTypedOuter<UBlueprint>();
+		UBlueprintGeneratedClass* BPGC = OwnerBlueprint ? Cast<UBlueprintGeneratedClass>(OwnerBlueprint->GeneratedClass) : nullptr;
+		if (BPGC)
+		{
+			const bool bIsEventGraph = BPGC->UberGraphFunction && FBlueprintEditorUtils::IsEventGraph(Graph);
+			FName GraphName = bIsEventGraph ? BPGC->UberGraphFunction->GetFName() : Graph->GetFName();
+			ScopedFunctionName = FName(*FString::Printf(TEXT("%s::%s"), *BPGC->GetName(), *GraphName.ToString()));
+		}
+	}
+	check (ScopedFunctionName != NAME_None);
+	return ScopedFunctionName;
+}
+
 //////////////////////////////////////////////////////////////////////
 // FBlueprintFunctionContext
 
@@ -420,11 +472,13 @@ void FBlueprintFunctionContext::InitialiseContextFromGraph(TSharedPtr<FBlueprint
 			// Instantiate context
 			OwningBlueprint = Blueprint;
 			BlueprintClass = BPClass;
-			const bool bInheritedClass = BlueprintContextIn->GetBlueprintClass() != BPClass;
+			bIsInheritedContext = BlueprintContextIn->GetBlueprintClass() != BPClass;
 			const bool bEventGraph = Graph->GetFName() == UEdGraphSchema_K2::GN_EventGraph;
 			GraphName = Graph->GetFName();
 			FunctionName = FunctionNameIn;
-			Function = BPClass->FindFunctionByName(FunctionName);
+			// Find Function
+			const FName UFunctionName = BlueprintContextIn->GetFunctionNameFromGraph(Graph);
+			Function = BPClass->FindFunctionByName(UFunctionName);
 			Function = Function.IsValid() ? Function : BPClass->UberGraphFunction;
 			if (bEventGraph)
 			{
@@ -436,31 +490,32 @@ void FBlueprintFunctionContext::InitialiseContextFromGraph(TSharedPtr<FBlueprint
 				}
 				for (auto EventNode : GraphEventNodes)
 				{
-					FName EventName = EventNode->GetFunctionName();
-					if (!BlueprintContext.Pin()->IsEventMapped(EventName))
+					const FName EventName = GetScopedEventName(EventNode->GetFunctionName());
+					FScriptExecNodeParams EventParams;
+					EventParams.SampleFrequency = 1;
+					EventParams.NodeName = EventName;
+					EventParams.ObservedObject = EventNode;
+					if (bIsInheritedContext)
 					{
-						FScriptExecNodeParams EventParams;
-						EventParams.NodeName = EventName;
-						EventParams.ObservedObject = EventNode;
-						EventParams.DisplayName = EventNode->GetNodeTitle(ENodeTitleType::ListView);
-						if (bInheritedClass)
-						{
-							EventParams.Tooltip = LOCTEXT("NavigateToInheritedEventLocationHyperlink_ToolTip", "Navigate to the Inherited Event");
-							EventParams.IconColor = FLinearColor(0.91f, 0.16f, 0.16f, 0.4f);
-						}
-						else
-						{
-							EventParams.Tooltip = LOCTEXT("NavigateToEventLocationHyperlink_ToolTip", "Navigate to the Event");
-							EventParams.IconColor = FLinearColor(0.91f, 0.16f, 0.16f);
-						}
+						EventParams.Tooltip = LOCTEXT("NavigateToInheritedEventLocationHyperlink_ToolTip", "Navigate to the Inherited Event");
+						EventParams.IconColor = GetDefault<UGraphEditorSettings>()->ParentFunctionCallNodeTitleColor;
+						EventParams.NodeFlags = EScriptExecutionNodeFlags::Event|EScriptExecutionNodeFlags::InheritedEvent;
+						EventParams.DisplayName = FText::FromString(FString::Printf(TEXT("%s (%s)"), *EventNode->GetNodeTitle(ENodeTitleType::ListView).ToString(), *OwningBlueprint->GetName()));
+					}
+					else
+					{
+						EventParams.Tooltip = LOCTEXT("NavigateToEventLocationHyperlink_ToolTip", "Navigate to the Event");
+						EventParams.IconColor = GetDefault<UGraphEditorSettings>()->EventNodeTitleColor;
 						EventParams.NodeFlags = EScriptExecutionNodeFlags::Event;
-						const FSlateBrush* EventIcon = EventNode->ShowPaletteIconOnNode() ?	EventNode->GetIconAndTint(EventParams.IconColor).GetOptionalIcon() :
-																							FEditorStyle::GetBrush(TEXT("BlueprintProfiler.BPNode"));
-						EventParams.Icon = const_cast<FSlateBrush*>(EventIcon);
-						TSharedPtr<FScriptExecutionNode> EventExecNode = CreateExecutionNode(EventParams);
-						AddEntryPoint(EventExecNode);
-						BlueprintContextIn->AddEventNode(AsShared(), EventExecNode);
+						EventParams.DisplayName = EventNode->GetNodeTitle(ENodeTitleType::ListView);
 					}
+					GetNodeCustomizations(EventParams);
+					const FSlateBrush* EventIcon = EventNode->ShowPaletteIconOnNode() ?	EventNode->GetIconAndTint(EventParams.IconColor).GetOptionalIcon() :
+																						FEditorStyle::GetBrush(TEXT("BlueprintProfiler.BPNode"));
+					EventParams.Icon = const_cast<FSlateBrush*>(EventIcon);
+					TSharedPtr<FScriptExecutionNode> EventExecNode = CreateExecutionNode(EventParams);
+					AddEntryPoint(EventExecNode);
+					BlueprintContextIn->AddEventNode(AsShared(), EventExecNode);
 				}
 				// Create any compiler generated events
 				FBlueprintDebugData& DebugData = BlueprintClass.Get()->GetDebugData();
@@ -471,31 +526,35 @@ void FBlueprintFunctionContext::InitialiseContextFromGraph(TSharedPtr<FBlueprint
 				// Map execution paths for each function entry nodes
 				TArray<UK2Node_FunctionEntry*> FunctionEntryNodes;
 				Graph->GetNodesOfClass<UK2Node_FunctionEntry>(FunctionEntryNodes);
-				const bool bConstructionScript = FunctionName == UEdGraphSchema_K2::FN_UserConstructionScript;
 				// Create function node
 				for (auto FunctionEntry : FunctionEntryNodes)
 				{
 					FScriptExecNodeParams FunctionNodeParams;
+					FunctionNodeParams.SampleFrequency = 1;
 					FunctionNodeParams.NodeName = FunctionName;
 					FunctionNodeParams.ObservedObject = FunctionEntry;
-					FunctionNodeParams.DisplayName = FText::FromName(FunctionName);
-					if (bInheritedClass)
+					if (bIsInheritedContext)
 					{
 						FunctionNodeParams.Tooltip = LOCTEXT("NavigateToInheritedFunctionLocationHyperlink_ToolTip", "Navigate to the Inherited Function");
-						FunctionNodeParams.IconColor = FLinearColor(0.91f, 0.16f, 0.16f, 0.4f);
+						FunctionNodeParams.IconColor = GetDefault<UGraphEditorSettings>()->ParentFunctionCallNodeTitleColor;
+						FunctionNodeParams.NodeFlags = EScriptExecutionNodeFlags::Event|EScriptExecutionNodeFlags::InheritedEvent;
+						FunctionNodeParams.DisplayName = FText::FromString(FString::Printf(TEXT("%s (%s)"), *GraphName.ToString(), *OwningBlueprint->GetName()));
 					}
 					else
 					{
 						FunctionNodeParams.Tooltip = LOCTEXT("NavigateToFunctionLocationHyperlink_ToolTip", "Navigate to the Function");
-						FunctionNodeParams.IconColor = FLinearColor(0.91f, 0.16f, 0.16f);
+						FunctionNodeParams.IconColor = GetDefault<UGraphEditorSettings>()->EventNodeTitleColor;
+						FunctionNodeParams.NodeFlags = EScriptExecutionNodeFlags::Event;
+						FunctionNodeParams.DisplayName = FText::FromName(GraphName);
 					}
-					FunctionNodeParams.NodeFlags = EScriptExecutionNodeFlags::Event;
+					GetNodeCustomizations(FunctionNodeParams);
 					const FSlateBrush* Icon = FunctionEntry->ShowPaletteIconOnNode() ? FunctionEntry->GetIconAndTint(FunctionNodeParams.IconColor).GetOptionalIcon() :
 																						FEditorStyle::GetBrush(TEXT("BlueprintProfiler.BPNode"));
 					FunctionNodeParams.Icon = const_cast<FSlateBrush*>(Icon);
 					TSharedPtr<FScriptExecutionNode> FunctionEntryNode = CreateExecutionNode(FunctionNodeParams);
 					AddEntryPoint(FunctionEntryNode);
-					if (bConstructionScript)
+					// Add user construction scripts as events
+					if (GraphName == UEdGraphSchema_K2::FN_UserConstructionScript)
 					{
 						BlueprintContextIn->AddEventNode(AsShared(), FunctionEntryNode);
 					}
@@ -581,13 +640,15 @@ bool FBlueprintFunctionContext::DetectCyclicLinks(TSharedPtr<FScriptExecutionNod
 			{
 				// Replace child with cyclic link node.
 				FScriptExecNodeParams CycleLinkParams;
+				CycleLinkParams.SampleFrequency = 1;
 				const FString NodeName = FString::Printf(TEXT("CyclicLinkTo_%i_%s"), ExecutionNodes.Num(), *Child->GetName().ToString());
 				CycleLinkParams.NodeName = FName(*NodeName);
 				CycleLinkParams.ObservedObject = Child->GetObservedObject();
 				CycleLinkParams.DisplayName = Child->GetDisplayName();
 				CycleLinkParams.Tooltip = LOCTEXT("CyclicLink_ToolTip", "Cyclic Link");
 				CycleLinkParams.NodeFlags = EScriptExecutionNodeFlags::CyclicLinkage|Child->GetFlags();
-				CycleLinkParams.IconColor = FLinearColor(0.91f, 0.16f, 0.16f);
+				CycleLinkParams.IconColor = Child->GetIconColor();
+				CycleLinkParams.IconColor.A = 0.15f;
 				const FSlateBrush* LinkIcon = Child->GetIcon();
 				CycleLinkParams.Icon = const_cast<FSlateBrush*>(LinkIcon);
 				Child = CreateExecutionNode(CycleLinkParams);
@@ -613,13 +674,15 @@ bool FBlueprintFunctionContext::DetectCyclicLinks(TSharedPtr<FScriptExecutionNod
 			{
 				// Break links and flag cycle linkage.
 				FScriptExecNodeParams CycleLinkParams;
+				CycleLinkParams.SampleFrequency = 1;
 				const FString NodeName = FString::Printf(TEXT("CyclicLinkTo_%i_%s"), ExecutionNodes.Num(), *LinkedExecNode->GetName().ToString());
 				CycleLinkParams.NodeName = FName(*NodeName);
 				CycleLinkParams.ObservedObject = LinkedExecNode->GetObservedObject();
 				CycleLinkParams.DisplayName = LinkedExecNode->GetDisplayName();
 				CycleLinkParams.Tooltip = LOCTEXT("CyclicLink_ToolTip", "Cyclic Link");
 				CycleLinkParams.NodeFlags = EScriptExecutionNodeFlags::CyclicLinkage|EScriptExecutionNodeFlags::ExecPin|EScriptExecutionNodeFlags::InvalidTrace;
-				CycleLinkParams.IconColor = FLinearColor(0.91f, 0.16f, 0.16f);
+				CycleLinkParams.IconColor = LinkedExecNode->GetIconColor();
+				CycleLinkParams.IconColor.A = 0.15f;
 				const FSlateBrush* LinkIcon = LinkedExecNode->GetIcon();
 				CycleLinkParams.Icon = const_cast<FSlateBrush*>(LinkIcon);
 				TSharedPtr<FScriptExecutionNode> NewLink = CreateExecutionNode(CycleLinkParams);
@@ -646,6 +709,7 @@ void FBlueprintFunctionContext::CreateDelegatePinEvents(TSharedPtr<FBlueprintExe
 	if (PinEvents.Num())
 	{
 		TMap<const UEdGraphNode*, TArray<FPinDelegateDesc>> NodeEventDescs;
+		TSharedPtr<FBlueprintFunctionContext> FunctionContext = AsShared();
 		// Build event contexts per node
 		for (auto PinEvent : PinEvents)
 		{
@@ -653,7 +717,8 @@ void FBlueprintFunctionContext::CreateDelegatePinEvents(TSharedPtr<FBlueprintExe
 			{
 				const UEdGraphNode* OwningNode = DelegatePin->GetOwningNode();
 				TArray<FPinDelegateDesc>& Events = NodeEventDescs.FindOrAdd(OwningNode);
-				Events.Add(FPinDelegateDesc(PinEvent.Key, DelegatePin));
+				const FName ScopedEventName = GetScopedEventName(PinEvent.Key);
+				Events.Add(FPinDelegateDesc(ScopedEventName, DelegatePin));
 			}
 		}
 		// Generate the event exec nodes
@@ -673,32 +738,34 @@ void FBlueprintFunctionContext::CreateDelegatePinEvents(TSharedPtr<FBlueprintExe
 			{
 				// Setup the basic exec node params.
 				FScriptExecNodeParams EventParams;
-				EventParams.NodeName = FName(*FString::Printf(TEXT("%s__DummyEvent"), *NodeEvents.Key->GetName()));
+				EventParams.SampleFrequency = 1;
+				EventParams.NodeName = FName(*FString::Printf(TEXT("%s::DummyEvent"), *NodeEvents.Key->GetName()));
 				EventParams.ObservedObject = NodeEvents.Key;
 				EventParams.OwningGraphName = NAME_None;
 				EventParams.DisplayName = NodeEvents.Key->GetNodeTitle(ENodeTitleType::ListView);
 				EventParams.Tooltip = LOCTEXT("NavigateToEventLocationHyperlink_ToolTip", "Navigate to the Event");
-				EventParams.NodeFlags = EScriptExecutionNodeFlags::Event;
-				EventParams.IconColor = FLinearColor(0.91f, 0.16f, 0.16f);
+				EventParams.NodeFlags = FunctionContext->IsInheritedContext() ? (EScriptExecutionNodeFlags::Event|EScriptExecutionNodeFlags::InheritedEvent) : EScriptExecutionNodeFlags::Event;
+				GetNodeCustomizations(EventParams);
 				const FSlateBrush* EventIcon = NodeEvents.Key->ShowPaletteIconOnNode() ? NodeEvents.Key->GetIconAndTint(EventParams.IconColor).GetOptionalIcon() :
 																						 FEditorStyle::GetBrush(TEXT("BlueprintProfiler.BPNode"));
 				EventParams.Icon = const_cast<FSlateBrush*>(EventIcon);
 				TSharedPtr<FScriptExecutionNode> EventExecNode = CreateExecutionNode(EventParams);
 				// Add entry points.
 				AddEntryPoint(EventExecNode);
-				BlueprintContextIn->AddEventNode(AsShared(), EventExecNode);
+				BlueprintContextIn->AddEventNode(FunctionContext, EventExecNode);
 			}
 			// Create the events for the pins
 			for (FPinDelegateDesc EventDesc : NodeEvents.Value)
 			{
 				FScriptExecNodeParams PinParams;
+				PinParams.SampleFrequency = 1;
 				PinParams.NodeName = EventDesc.EventName;
 				PinParams.ObservedPin = EventDesc.DelegatePin;
 				PinParams.OwningGraphName = NAME_None;
 				PinParams.DisplayName = EventDesc.DelegatePin->GetDisplayName();
 				PinParams.Tooltip = LOCTEXT("ExecPin_ExpandExecutionPath_ToolTip", "Expand execution path");
 				PinParams.NodeFlags = EScriptExecutionNodeFlags::ExecPin|EScriptExecutionNodeFlags::EventPin;
-				PinParams.IconColor = FLinearColor(1.f, 1.f, 1.f, 0.8f);
+				PinParams.IconColor = GetDefault<UGraphEditorSettings>()->DelegatePinTypeColor;
 				const bool bPinLinked = EventDesc.DelegatePin->LinkedTo.Num() > 0;
 				const FSlateBrush* Icon = bPinLinked ?	FEditorStyle::GetBrush(TEXT("BlueprintProfiler.BPPinConnected")) : 
 														FEditorStyle::GetBrush(TEXT("BlueprintProfiler.BPPinDisconnected"));
@@ -706,7 +773,7 @@ void FBlueprintFunctionContext::CreateDelegatePinEvents(TSharedPtr<FBlueprintExe
 				TSharedPtr<FScriptExecutionNode> PinExecNode = CreateExecutionNode(PinParams);
 				AddEntryPoint(PinExecNode);
 				// Register the function context as a handler for the event.
-				BlueprintContextIn->RegisterEventContext(EventDesc.EventName, AsShared());
+				BlueprintContextIn->RegisterEventContext(EventDesc.EventName, FunctionContext);
 				// Register exec node under pin name
 				const FName PinName = GetUniquePinName(EventDesc.DelegatePin);
 				ExecutionNodes.Add(PinName) = PinExecNode;
@@ -715,54 +782,140 @@ void FBlueprintFunctionContext::CreateDelegatePinEvents(TSharedPtr<FBlueprintExe
 	}
 }
 
+void FBlueprintFunctionContext::GetNodeCustomizations(FScriptExecNodeParams& ParamsInOut) const
+{
+	// Pick a color based on flags.
+	if ((ParamsInOut.NodeFlags & (EScriptExecutionNodeFlags::InheritedEvent|EScriptExecutionNodeFlags::ParentFunctionCall)) != 0U)
+	{
+		// Inherited events and calls
+		ParamsInOut.IconColor = GetDefault<UGraphEditorSettings>()->ParentFunctionCallNodeTitleColor;
+	}
+	else if ((ParamsInOut.NodeFlags & EScriptExecutionNodeFlags::Event) != 0U)
+	{
+		// Events and custom events
+		ParamsInOut.IconColor = GetDefault<UGraphEditorSettings>()->EventNodeTitleColor;
+	}
+	else if ((ParamsInOut.NodeFlags & EScriptExecutionNodeFlags::FunctionCall) != 0U)
+	{
+		// Function calls
+		ParamsInOut.IconColor = GetDefault<UGraphEditorSettings>()->FunctionCallNodeTitleColor;
+	}
+	else
+	{
+		// Set as the default node color.
+		ParamsInOut.IconColor = FLinearColor(1.f, 1.f, 1.f, 0.8f);
+		// Check for any final specialisations
+		if ((ParamsInOut.NodeFlags & (EScriptExecutionNodeFlags::PureNode|EScriptExecutionNodeFlags::PureChain)) != 0U)
+		{
+			// Pure nodes
+			ParamsInOut.IconColor = GetDefault<UGraphEditorSettings>()->PureFunctionCallNodeTitleColor;
+		}
+		else if ((ParamsInOut.NodeFlags & EScriptExecutionNodeFlags::SequentialBranch) != 0U)
+		{
+			// Sequential branches
+			ParamsInOut.IconColor = GetDefault<UGraphEditorSettings>()->ExecSequenceNodeTitleColor;
+		}
+		else if ((ParamsInOut.NodeFlags & EScriptExecutionNodeFlags::ConditionalBranch) != 0U)
+		{
+			// Consditional branches
+			ParamsInOut.IconColor = GetDefault<UGraphEditorSettings>()->ExecBranchNodeTitleColor;
+		}
+		else if (ParamsInOut.ObservedObject->IsA<UK2Node_Event>() || ParamsInOut.ObservedObject->IsA<UK2Node_FunctionEntry>())
+		{
+			// Differentiate between execution path and entry node when they have the same name
+			FFormatNamedArguments Args;
+			Args.Add(TEXT("NodeName"), ParamsInOut.DisplayName);
+			Args.Add(TEXT("EntryNode"), LOCTEXT("EntryNode", "Entry Node"));
+			ParamsInOut.DisplayName = FText::Format(FText::FromString("{NodeName} ({EntryNode})"), Args);
+		}
+	}
+}
+
+void FBlueprintFunctionContext::DetermineGraphNodeCharacteristics(const UEdGraphNode* GraphNode, TArray<UEdGraphPin*>& InputPins, TArray<UEdGraphPin*>& ExecPins, FScriptExecNodeParams& NodeParams)
+{
+	// Set the standard sample base
+	NodeParams.SampleFrequency = 1;
+	// Evaluate Execution and Input Pins
+	int32 ConnectedExecPins = 0;
+	for (auto Pin : GraphNode->Pins)
+	{
+		if (Pin->Direction == EGPD_Output && Pin->PinType.PinCategory == UEdGraphSchema_K2::PC_Exec)
+		{
+			ExecPins.Add(Pin);
+			if (Pin->LinkedTo.Num())
+			{
+				ConnectedExecPins++;
+			}
+		}
+		else if (Pin->Direction == EGPD_Input && Pin->PinType.PinCategory != UEdGraphSchema_K2::PC_Exec)
+		{
+			InputPins.Add(Pin);
+		}
+	}
+	// Identify branching types and pure nodes based on pin layout.
+	if (ExecPins.Num() == 0 && !GraphNode->IsA<UK2Node_FunctionResult>())
+	{
+		NodeParams.NodeFlags |= EScriptExecutionNodeFlags::PureNode;
+	}
+	else if (ExecPins.Num() > 1)
+	{
+		if (GraphNode->IsA<UK2Node_ExecutionSequence>())
+		{
+			// Update the node params with the expected base sample rate.
+			NodeParams.SampleFrequency = ConnectedExecPins;
+			NodeParams.NodeFlags |= EScriptExecutionNodeFlags::SequentialBranch;
+		}
+		else
+		{
+			NodeParams.NodeFlags |= EScriptExecutionNodeFlags::ConditionalBranch;
+		}
+	}
+	// Identify function calls and custom events
+	if (GraphNode->IsA<UK2Node_CallParentFunction>())
+	{
+		NodeParams.NodeFlags |= EScriptExecutionNodeFlags::ParentFunctionCall;
+	}
+	else if (GraphNode->IsA<UK2Node_CallFunction>())
+	{
+		NodeParams.NodeFlags |= EScriptExecutionNodeFlags::FunctionCall;
+	}
+	else if (GraphNode->IsA<UK2Node_CustomEvent>())
+	{
+		NodeParams.NodeFlags |= EScriptExecutionNodeFlags::CustomEvent;
+	}
+}
+
 TSharedPtr<FScriptExecutionNode> FBlueprintFunctionContext::MapNodeExecution(UEdGraphNode* NodeToMap)
 {
 	TSharedPtr<FScriptExecutionNode> MappedNode;
 	if (NodeToMap)
 	{
+		// Lookup existing mapped node
 		MappedNode = GetProfilerDataForGraphNode(NodeToMap);
+		// Map if not existing.
 		if (!MappedNode.IsValid())
 		{
+			// Determine node characteristics
 			FScriptExecNodeParams NodeParams;
+			NodeParams.NodeFlags = EScriptExecutionNodeFlags::Node;
+			TArray<UEdGraphPin*> ExecPins;
+			TArray<UEdGraphPin*> InputPins;
+			DetermineGraphNodeCharacteristics(NodeToMap, InputPins, ExecPins, NodeParams);
 			NodeParams.NodeName = NodeToMap->GetFName();
 			NodeParams.ObservedObject = NodeToMap;
 			NodeParams.DisplayName = NodeToMap->GetNodeTitle(ENodeTitleType::ListView);
 			NodeParams.Tooltip = LOCTEXT("NavigateToNodeLocationHyperlink_ToolTip", "Navigate to the Node");
-			NodeParams.NodeFlags = EScriptExecutionNodeFlags::Node;
-			NodeParams.IconColor = FLinearColor(1.f, 1.f, 1.f, 0.8f);
+
+			GetNodeCustomizations(NodeParams);
 			const FSlateBrush* NodeIcon = NodeToMap->ShowPaletteIconOnNode() ? NodeToMap->GetIconAndTint(NodeParams.IconColor).GetOptionalIcon() :
 																				FEditorStyle::GetBrush(TEXT("BlueprintProfiler.BPNode"));
 			NodeParams.Icon = const_cast<FSlateBrush*>(NodeIcon);
 			MappedNode = CreateExecutionNode(NodeParams);
-			// Evaluate Execution and Input Pins
-			TArray<UEdGraphPin*> ExecPins;
-			TArray<UEdGraphPin*> InputPins;
-			for (auto Pin : NodeToMap->Pins)
-			{
-				if (Pin->Direction == EGPD_Output && Pin->PinType.PinCategory == UEdGraphSchema_K2::PC_Exec)
-				{
-					ExecPins.Add(Pin);
-				}
-				else if (Pin->Direction == EGPD_Input && Pin->PinType.PinCategory != UEdGraphSchema_K2::PC_Exec)
-				{
-					InputPins.Add(Pin);
-				}
-			}
-			// Flag as pure node or build pure node script range.
-			if (ExecPins.Num() == 0)
-			{
-				MappedNode->AddFlags(EScriptExecutionNodeFlags::PureNode);
-			}
-			else
+			// Discover pure node script ranges
+			if (!MappedNode->IsPureNode())
 			{
 				MappedNode->SetPureNodeScriptCodeRange(GetPureNodeScriptCodeRange(NodeToMap));
 			}
-			// Determine a few node characteristics to guide processing.
-			if (ExecPins.Num() > 1)
-			{
-				const bool bExecSequence = NodeToMap->IsA<UK2Node_ExecutionSequence>();
-				MappedNode->AddFlags(bExecSequence ? EScriptExecutionNodeFlags::SequentialBranch : EScriptExecutionNodeFlags::ConditionalBranch);
-			}
 			// Evaluate non-exec input pins (pure node execution chains)
 			if (InputPins.Num())
 			{
@@ -773,39 +926,91 @@ TSharedPtr<FScriptExecutionNode> FBlueprintFunctionContext::MapNodeExecution(UEd
 			{
 				MapExecPins(MappedNode, ExecPins);
 			}
-			// Evaluate Children for call sites
-			if (UK2Node_CallFunction* FunctionCallSite = Cast<UK2Node_CallFunction>(NodeToMap))
+			// Locate any pure nodes from the root impure node.
+			if (!MappedNode->IsPureNode())
 			{
-				const UEdGraphNode* FunctionNode = nullptr; 
-				if (UEdGraph* CalledGraph = FunctionCallSite->GetFunctionGraph(FunctionNode))
+				TSharedPtr<FScriptExecutionNode> PureChainNode = MappedNode->GetPureChainNode();
+				if (PureChainNode.IsValid())
 				{
-					// Update Exec node
-					const bool bEventCall = FunctionNode ? FunctionNode->IsA<UK2Node_Event>() : false;
-					MappedNode->AddFlags(EScriptExecutionNodeFlags::FunctionCall);
-					MappedNode->SetIconColor(FLinearColor(0.46f, 0.54f, 0.95f));
-					MappedNode->SetToolTipText(LOCTEXT("NavigateToFunctionCallsiteHyperlink_ToolTip", "Navigate to the Function Callsite"));
-					// Don't add entry points for events.
-					if (!bEventCall)
+					TMap<int32, TSharedPtr<FScriptExecutionNode>> AllPureNodes;
+					PureChainNode->GetAllPureNodes(AllPureNodes);
+					// Patch pure node chains into function context for lookup later.
+					for (auto PureNode : AllPureNodes)
 					{
-						// Update the function context
-						const FName FunctionGraphName = CalledGraph->GetFName();
-						TSharedPtr<FBlueprintFunctionContext> NewFunctionContext = BlueprintContext.Pin()->GetFunctionContext(FunctionGraphName);
-						if (NewFunctionContext.IsValid())
+						if (!ExecutionNodes.Contains(PureNode.Value->GetName()))
 						{
-							ChildFunctionContexts.Add(FunctionGraphName) = NewFunctionContext;
+							ExecutionNodes.Add(PureNode.Value->GetName()) = PureNode.Value;
 						}
 					}
 				}
 			}
-			else if (UK2Node_CustomEvent* CustomEvent = Cast<UK2Node_CustomEvent>(NodeToMap))
+			// Evaluate Children for call sites
+			if (MappedNode->IsFunctionCallSite()||MappedNode->IsParentFunctionCallSite())
 			{
-				MappedNode->AddFlags(EScriptExecutionNodeFlags::CustomEvent);
+				if (UK2Node_CallFunction* FunctionCallSite = Cast<UK2Node_CallFunction>(NodeToMap))
+				{
+					const UEdGraphNode* FunctionNode = nullptr; 
+					if (UEdGraph* CalledGraph = FunctionCallSite->GetFunctionGraph(FunctionNode))
+					{
+						// Update Exec node
+						const bool bEventCall = FunctionNode ? FunctionNode->IsA<UK2Node_Event>() : false;
+						MappedNode->SetToolTipText(LOCTEXT("NavigateToFunctionCallsiteHyperlink_ToolTip", "Navigate to the Function Callsite"));
+						// Don't add entry points for events.
+						if (!bEventCall)
+						{
+							// Update the function context
+							TSharedPtr<FBlueprintFunctionContext> NewFunctionContext = BlueprintContext.Pin()->GetFunctionContextFromGraph(CalledGraph);
+							if (NewFunctionContext.IsValid())
+							{
+								NewFunctionContext->AddCallSiteEntryPointsToNode(MappedNode);
+								ChildFunctionContexts.Add(NewFunctionContext->GetFunctionName()) = NewFunctionContext;
+							}
+						}
+					}
+				}
 			}
 		}
 	}
 	return MappedNode;
 }
 
+TSharedPtr<FScriptExecutionNode> FBlueprintFunctionContext::FindOrCreatePureChainRoot(TSharedPtr<FScriptExecutionNode> ExecNode)
+{
+	TSharedPtr<FScriptExecutionNode> PureChainRootNode;
+	// Add a pure chain container node as the root, if it's not already in place.
+	if (ExecNode.IsValid() && !ExecNode->IsPureNode())
+	{
+		static const FString PureChainNodeNameSuffix = TEXT("__PROFILER_InputPureTime");
+		const FName PureChainNodeName = FName(*(ExecNode->GetName().ToString() + PureChainNodeNameSuffix));
+		// Try to find the root first
+		for (auto ChildNode : ExecNode->GetChildNodes())
+		{
+			if (ChildNode->GetName() == PureChainNodeName)
+			{
+				PureChainRootNode = ChildNode;
+			}
+		}
+		// Otherwise create one
+		if (!PureChainRootNode.IsValid())
+		{
+			FScriptExecNodeParams PureChainParams;
+			PureChainParams.SampleFrequency = 1;
+			PureChainParams.NodeName = PureChainNodeName;
+			PureChainParams.ObservedObject = ExecNode->GetObservedObject();
+			PureChainParams.DisplayName = LOCTEXT("PureChain_DisplayName", "Pure Time");
+			PureChainParams.Tooltip = LOCTEXT("PureChain_ToolTip", "Expand pure node timing");
+			PureChainParams.NodeFlags = EScriptExecutionNodeFlags::PureChain;
+			GetNodeCustomizations(PureChainParams);
+			const FSlateBrush* Icon = FEditorStyle::GetBrush(TEXT("BlueprintProfiler.PureNode"));
+			PureChainParams.Icon = const_cast<FSlateBrush*>(Icon);
+			PureChainRootNode = CreateExecutionNode(PureChainParams);
+			PureChainRootNode->SetPureNodeScriptCodeRange(ExecNode->GetPureNodeScriptCodeRange());
+			ExecNode->AddChildNode(PureChainRootNode);
+		}
+	}
+	return PureChainRootNode;
+}
+
 void FBlueprintFunctionContext::MapInputPins(TSharedPtr<FScriptExecutionNode> ExecNode, const TArray<UEdGraphPin*>& Pins)
 {
 	TArray<int32> PinScriptCodeOffsets;
@@ -819,6 +1024,38 @@ void FBlueprintFunctionContext::MapInputPins(TSharedPtr<FScriptExecutionNode> Ex
 			// Pass through non-relevant (e.g. reroute) nodes.
 			UK2Node* OwningNode = FBlueprintEditorUtils::FindFirstCompilerRelevantNode(LinkedPin);
 
+			// If this is a tunnel node - we need to map through the tunnel
+			if (OwningNode && OwningNode->IsA<UK2Node_Tunnel>())
+			{
+				if (FBlueprintEditorUtils::IsTunnelInstanceNode(OwningNode))
+				{
+					// Find the tunnel instance context
+					const FName ScopedTunnelContextName = GetTunnelInstanceFunctionName(OwningNode);
+					TSharedPtr<FBlueprintTunnelInstanceContext> TunnelContext = StaticCastSharedPtr<FBlueprintTunnelInstanceContext>(BlueprintContext.Pin()->GetFunctionContext(ScopedTunnelContextName));
+					check (TunnelContext.IsValid());
+					TSharedPtr<FScriptExecutionNode> TunnelPureNode = TunnelContext->GetPureNodeFromTunnelPin(LinkedPin);
+					// Add the linked pure nodes
+					if (TunnelPureNode.IsValid() && TunnelPureNode->IsPureChain())
+					{
+						// Grab the pure chain root.
+						PureChainRootNode = FindOrCreatePureChainRoot(ExecNode);
+						if (!PureChainRootNode.IsValid())
+						{
+							PureChainRootNode = ExecNode;
+						}
+						// Copy relevant pure chain links
+						FInt32Range TunnelPureRange = PureChainRootNode->GetPureNodeScriptCodeRange();
+						for (auto LinkedNode : TunnelPureNode->GetLinkedNodes())
+						{
+							if (TunnelPureRange.Contains(LinkedNode.Key))
+							{
+								PureChainRootNode->AddLinkedNode(LinkedNode.Key, LinkedNode.Value);
+							}
+						}
+					}
+				}
+			}
+
 			// Note: Intermediate pure nodes can have output pins that masquerade as impure node output pins when links are "moved" from the source graph (thus
 			// resulting in a false association here with one or more script code offsets), so we must first ensure that the link is really to a pure node output.
 			if (OwningNode && OwningNode->IsNodePure())
@@ -829,19 +1066,7 @@ void FBlueprintFunctionContext::MapInputPins(TSharedPtr<FScriptExecutionNode> Ex
 					// Add a pure chain container node as the root, if it's not already in place.
 					if (PureChainRootNode == ExecNode && !ExecNode->IsPureNode())
 					{
-						FScriptExecNodeParams PureChainParams;
-						static const FString PureChainNodeNameSuffix = TEXT("__PROFILER_PureTime");
-						FString PureChainNodeNameString = ExecNode->GetName().ToString() + PureChainNodeNameSuffix;
-						PureChainParams.NodeName = FName(*PureChainNodeNameString);
-						PureChainParams.DisplayName = LOCTEXT("PureChain_DisplayName", "Pure Time");
-						PureChainParams.Tooltip = LOCTEXT("PureChain_ToolTip", "Expand pure node timing");
-						PureChainParams.NodeFlags = EScriptExecutionNodeFlags::PureChain;
-						PureChainParams.IconColor = FLinearColor(0.2f, 1.f, 0.2f);
-						const FSlateBrush* Icon = FEditorStyle::GetBrush(TEXT("BlueprintProfiler.PureNode"));
-						PureChainParams.Icon = const_cast<FSlateBrush*>(Icon);
-						PureChainRootNode = CreateExecutionNode(PureChainParams);
-						PureChainRootNode->SetPureNodeScriptCodeRange(ExecNode->GetPureNodeScriptCodeRange());
-						ExecNode->AddChildNode(PureChainRootNode);
+						PureChainRootNode = FindOrCreatePureChainRoot(ExecNode);
 					}
 
 					TSharedPtr<FScriptExecutionNode> PureNode = MapNodeExecution(OwningNode);
@@ -883,6 +1108,7 @@ void FBlueprintFunctionContext::MapExecPins(TSharedPtr<FScriptExecutionNode> Exe
 		if (bBranchedExecution)
 		{
 			FScriptExecNodeParams LinkNodeParams;
+			LinkNodeParams.SampleFrequency = 1;
 			LinkNodeParams.NodeName = PinName;
 			LinkNodeParams.ObservedObject = Pin->GetOwningNode();
 			LinkNodeParams.DisplayName = Pin->GetDisplayName();
@@ -916,19 +1142,41 @@ void FBlueprintFunctionContext::MapExecPins(TSharedPtr<FScriptExecutionNode> Exe
 	}
 }
 
+TSharedPtr<FScriptExecutionNode> FBlueprintFunctionContext::GetTunnelBoundaryNodeChecked(const UEdGraphPin* TunnelPin)
+{
+	TSharedPtr<FScriptExecutionNode> Result;
+	if (const UK2Node_Tunnel* TunnelNode = Cast<UK2Node_Tunnel>(TunnelPin->GetOwningNode()))
+	{
+		if (FBlueprintEditorUtils::IsTunnelInstanceNode(TunnelNode))
+		{
+			// Lookup external tunnel boundary.
+			const FName TunnelInstanceFunctionName = GetTunnelInstanceFunctionName(TunnelNode);
+			TSharedPtr<FBlueprintFunctionContext> TunnelContext = BlueprintContext.Pin()->GetFunctionContext(TunnelInstanceFunctionName);
+			if (TunnelContext.IsValid())
+			{
+				Result = TunnelContext->GetProfilerDataForNode(GetTunnelBoundaryName(TunnelPin));
+			}
+		}
+		else
+		{
+			// Internal boundary lookup, faster path.
+			Result = GetProfilerDataForNode(GetPinName(TunnelPin));
+		}
+	}
+	check (Result.IsValid());
+	return Result;
+}
+
 TSharedPtr<FScriptExecutionNode> FBlueprintFunctionContext::MapTunnelBoundary(const UEdGraphPin* TunnelPin)
 {
 	TSharedPtr<FScriptExecutionNode> TunnelBoundaryNode;
 	if (TunnelPin)
 	{
-		const FName TunnelBoundaryName = GetTunnelBoundaryName(TunnelPin);
-		TunnelBoundaryNode = GetProfilerDataForNodeChecked(TunnelBoundaryName);
+		TunnelBoundaryNode = GetTunnelBoundaryNodeChecked(TunnelPin);
 		if (TunnelBoundaryNode->IsTunnelEntry())
 		{
 			TSharedPtr<FScriptExecutionTunnelEntry> TunnelEntryInstance = StaticCastSharedPtr<FScriptExecutionTunnelEntry>(TunnelBoundaryNode);
-			UEdGraphNode* TunnelInstanceNode = TunnelPin->GetOwningNode();
-
-			for (auto ExitSite : TunnelEntryInstance->GetExitSites())
+			for (auto ExitSite : TunnelEntryInstance->GetLinkedNodes())
 			{
 				if (ExitSite.Value->IsTunnelExit())
 				{
@@ -982,7 +1230,7 @@ FName FBlueprintFunctionContext::GetPinName(const UEdGraphPin* Pin)
 		if (OwningNode->IsA<UK2Node_Tunnel>())
 		{
 			// Tunnel pins have to be unique so we need the node name in addition to pin name.
-			UEdGraph* TunnelGraph = GetTunnelGraphFromNode(OwningNode);
+			UEdGraph* TunnelGraph = FBlueprintFunctionContext::GetGraphFromNode(OwningNode, false);
 			TunnelGraph = TunnelGraph ? TunnelGraph : OwningNode->GetTypedOuter<UEdGraph>();
 			FName TunnelName = TunnelGraph->GetFName();
 			PinString = FString::Printf(TEXT("%s_%s"), *TunnelName.ToString(), *PinString);
@@ -1044,30 +1292,70 @@ UEdGraphPin* FBlueprintFunctionContext::FindMatchingPin(const UEdGraphNode* Node
 	return MatchingPin;
 }
 
-UEdGraph* FBlueprintFunctionContext::GetTunnelGraphFromNode(const UEdGraphNode* TunnelNode)
+UEdGraph* FBlueprintFunctionContext::GetGraphFromNode(const UEdGraphNode* GraphNode, const bool bAllowNonTunnel)
 {
-	UEdGraph* TunnelGraph = nullptr;
-	if (const UK2Node_MacroInstance* MacroInstance = Cast<UK2Node_MacroInstance>(TunnelNode))
+	UEdGraph* OwningGraph = nullptr;
+	if (GraphNode)
 	{
-		TunnelGraph = MacroInstance->GetMacroGraph();
+		if (const UK2Node_MacroInstance* MacroInstance = Cast<UK2Node_MacroInstance>(GraphNode))
+		{
+			OwningGraph = MacroInstance->GetMacroGraph();
+		}
+		else if (const UK2Node_Composite* CompositeNode = Cast<UK2Node_Composite>(GraphNode))
+		{
+			OwningGraph = CompositeNode->BoundGraph;
+		}
+		else if (bAllowNonTunnel)
+		{
+			OwningGraph = GraphNode->GetTypedOuter<UEdGraph>();
+		}
 	}
-	else if (const UK2Node_Composite* CompositeInstance = Cast<UK2Node_Composite>(TunnelNode))
+	return OwningGraph;
+}
+
+FName FBlueprintFunctionContext::GetTunnelInstanceFunctionName(const UEdGraphNode* GraphNode)
+{
+	FName ScopedFunctionName = NAME_None;
+	UEdGraph* OuterGraph = GraphNode ? GraphNode->GetTypedOuter<UEdGraph>() : nullptr;
+	if (OuterGraph)
 	{
-		TunnelGraph = CompositeInstance->BoundGraph;
+		// To identify instances exactly, we need the owning graph name and the node name to avoid collisions.
+		ScopedFunctionName = FName(*FString::Printf(TEXT("%s::%s"), *OuterGraph->GetName(), *GraphNode->GetName()));
 	}
-	return TunnelGraph;
+	check (ScopedFunctionName != NAME_None);
+	return ScopedFunctionName;
 }
 
-bool FBlueprintFunctionContext::IsTunnelInternal(const UEdGraphNode* TunnelNode)
+FName FBlueprintFunctionContext::GetScopedFunctionNameFromNode(const UEdGraphNode* GraphNode) const
 {
-	return !(TunnelNode->IsA<UK2Node_MacroInstance>()||TunnelNode->IsA<UK2Node_Composite>());
+	FName ScopedFunctionName = NAME_None;
+	if (GraphNode)
+	{
+		if (const UEdGraphNode* TunnelNode = GetTunnelNodeFromGraphNode(GraphNode))
+		{
+			ScopedFunctionName = GetTunnelInstanceFunctionName(TunnelNode);
+		}
+		else if (const UEdGraph* Graph = GetGraphFromNode(GraphNode))
+		{
+			UBlueprint* OwnerBlueprint = Graph->GetTypedOuter<UBlueprint>();
+			UBlueprintGeneratedClass* BPGC = OwnerBlueprint ? Cast<UBlueprintGeneratedClass>(OwnerBlueprint->GeneratedClass) : nullptr;
+			if (BPGC)
+			{
+				const bool bIsEventGraph = BPGC->UberGraphFunction && FBlueprintEditorUtils::IsEventGraph(Graph);
+				FName NodeGraphName = bIsEventGraph ? BPGC->UberGraphFunction->GetFName() : Graph->GetFName();
+				ScopedFunctionName = FName(*FString::Printf(TEXT("%s::%s"), *BPGC->GetName(), *NodeGraphName.ToString()));
+			}
+		}
+	}
+	check (ScopedFunctionName != NAME_None);
+	return ScopedFunctionName;
 }
 
 void FBlueprintFunctionContext::MapTunnelExits(TSharedPtr<FScriptExecutionTunnelEntry> TunnelEntryPoint)
 {
 	if (TunnelEntryPoint.IsValid())
 	{
-		for (auto ExitSite : TunnelEntryPoint->GetExitSites())
+		for (auto ExitSite : TunnelEntryPoint->GetLinkedNodes())
 		{
 			if (const UEdGraphPin* ExitPin = ExitSite.Value->GetObservedPin())
 			{
@@ -1083,13 +1371,20 @@ void FBlueprintFunctionContext::MapTunnelExits(TSharedPtr<FScriptExecutionTunnel
 
 bool FBlueprintFunctionContext::HasProfilerDataForNode(const FName NodeName) const
 {
-	return ExecutionNodes.Contains(NodeName);
+	TSharedPtr<FScriptExecutionNode> Result = GetProfilerDataForNode(NodeName);
+	return Result.IsValid();
+}
+
+bool FBlueprintFunctionContext::HasProfilerDataForGraphNode(const UEdGraphNode* Node) const
+{
+	TSharedPtr<FScriptExecutionNode> Result = GetProfilerDataForGraphNode(Node);
+	return Result.IsValid();
 }
 
-TSharedPtr<FScriptExecutionNode> FBlueprintFunctionContext::GetProfilerDataForNode(const FName NodeName)
+TSharedPtr<FScriptExecutionNode> FBlueprintFunctionContext::GetProfilerDataForNode(const FName NodeName) const
 {
 	TSharedPtr<FScriptExecutionNode> Result;
-	if (TSharedPtr<FScriptExecutionNode>* SearchResult = ExecutionNodes.Find(NodeName))
+	if (const TSharedPtr<FScriptExecutionNode>* SearchResult = ExecutionNodes.Find(NodeName))
 	{
 		Result = *SearchResult;
 	}
@@ -1111,40 +1406,87 @@ TSharedPtr<FScriptExecutionNode> FBlueprintFunctionContext::GetProfilerDataForNo
 	return Result;
 }
 
-TSharedPtr<FScriptExecutionNode> FBlueprintFunctionContext::GetProfilerDataForNodeChecked(const FName NodeName)
+TSharedPtr<FScriptExecutionNode> FBlueprintFunctionContext::GetProfilerDataForNodeChecked(const FName NodeName) const
 {
 	TSharedPtr<FScriptExecutionNode> Result = GetProfilerDataForNode(NodeName);
 	check (Result.IsValid());
 	return Result;
 }
 
-TSharedPtr<FScriptExecutionNode> FBlueprintFunctionContext::GetProfilerDataForGraphNode(const UEdGraphNode* Node)
+TSharedPtr<FScriptExecutionNode> FBlueprintFunctionContext::GetProfilerDataForGraphPin(const UEdGraphPin* Pin) const
 {
 	TSharedPtr<FScriptExecutionNode> Result;
-	if (Node)
+	if (Pin)
 	{
-		if (UEdGraph* OuterGraph = Node->GetTypedOuter<UEdGraph>())
+		// Do a local lookup, this could be a tunnel/macro node and the graph wouldn't match.
+		if (const TSharedPtr<FScriptExecutionNode>* SearchResult = ExecutionNodes.Find(GetPinName(Pin)))
+		{
+			Result = *SearchResult;
+		}
+		// Check objects match
+		if (Result.IsValid())
 		{
-			const bool bTunnelInstance = Node->IsA<UK2Node_MacroInstance>()||Node->IsA<UK2Node_Composite>();
-			// Correctly locate macro/tunnel graphs
-			const FName ObjectGraphName = bTunnelInstance ? Node->GetFName() : OuterGraph->GetFName();
-			if (ObjectGraphName == GraphName)
+			if (Pin != Result->GetObservedPin())
 			{
-				Result = GetProfilerDataForNode(Node->GetFName());
+				// If the observed object doesn't match we need to look elsewhere.
+				Result.Reset();
 			}
-			else if (TWeakPtr<FBlueprintFunctionContext>* FunctionContext = ChildFunctionContexts.Find(ObjectGraphName))
+		}
+		// Perform a more exhaustive search.
+		if (!Result.IsValid())
+		{
+			const FName ObjectFunctionName = GetScopedFunctionNameFromNode(Pin->GetOwningNode());
+			if (ObjectFunctionName != FunctionName)
 			{
-				if (FunctionContext->IsValid())
+				// Lookup function context otherwise search child contexts
+				if (const TWeakPtr<FBlueprintFunctionContext>* FunctionContext = ChildFunctionContexts.Find(ObjectFunctionName))
 				{
-					Result = FunctionContext->Pin()->GetProfilerDataForNode(Node->GetFName());
+					if (FunctionContext->IsValid())
+					{
+						Result = FunctionContext->Pin()->GetProfilerDataForGraphPin(Pin);
+					}
 				}
 			}
 		}
 	}
-	if (Result.IsValid() && Result->GetObservedObject() != Node)
+	return Result;
+}
+
+TSharedPtr<FScriptExecutionNode> FBlueprintFunctionContext::GetProfilerDataForGraphNode(const UEdGraphNode* Node) const
+{
+	TSharedPtr<FScriptExecutionNode> Result;
+	if (Node)
 	{
-		// Objects don't match, perform a more exhaustive search.
-		Result.Reset();
+		// Do a local lookup, this could be a tunnel/macro node and the graph wouldn't match.
+		if (const TSharedPtr<FScriptExecutionNode>* SearchResult = ExecutionNodes.Find(Node->GetFName()))
+		{
+			Result = *SearchResult;
+		}
+		// Check objects match
+		if (Result.IsValid())
+		{
+			if (Node != Result->GetObservedObject())
+			{
+				// If the observed object doesn't match we need to look elsewhere.
+				Result.Reset();
+			}
+		}
+		// Perform a more exhaustive search.
+		if (!Result.IsValid())
+		{
+			const FName ObjectFunctionName = GetScopedFunctionNameFromNode(Node);
+			if (ObjectFunctionName != FunctionName)
+			{
+				// Lookup function context otherwise search child contexts
+				if (const TWeakPtr<FBlueprintFunctionContext>* FunctionContext = ChildFunctionContexts.Find(ObjectFunctionName))
+				{
+					if (FunctionContext->IsValid())
+					{
+						Result = FunctionContext->Pin()->GetProfilerDataForGraphNode(Node);
+					}
+				}
+			}
+		}
 	}
 	return Result;
 }
@@ -1171,36 +1513,88 @@ template<typename ExecNodeType> TSharedPtr<ExecNodeType> FBlueprintFunctionConte
 	return Result;
 }
 
+bool FBlueprintFunctionContext::GetProfilerContextFromScriptOffset(const int32 ScriptOffset, TSharedPtr<FScriptExecutionNode>& ExecNode, TSharedPtr<FBlueprintFunctionContext>& FunctionContext)
+{
+	if (const UEdGraphNode* GraphNode = GetNodeFromCodeLocation(ScriptOffset))
+	{
+		// get the fully scoped function name
+		const FName FunctionContextName = GetScopedFunctionNameFromNode(GraphNode);
+		if (FunctionContextName == FunctionName)
+		{
+			FunctionContext = AsShared();
+			ExecNode = GetProfilerDataForNode(GraphNode->GetFName());
+		}
+		else
+		{
+			if (TWeakPtr<FBlueprintFunctionContext>* FunctionContextPtr = ChildFunctionContexts.Find(FunctionContextName))
+			{
+				if (FunctionContextPtr->IsValid())
+				{
+					// We now have the source tunnel graph context, use this to find the instance and then the node.
+					FunctionContextPtr->Pin()->GetProfilerContextFromScriptOffset(ScriptOffset, ExecNode, FunctionContext);
+				}
+			}
+		}
+	}
+	return ExecNode.IsValid() && FunctionContext.IsValid();
+}
+
 const UEdGraphNode* FBlueprintFunctionContext::GetNodeFromCodeLocation(const int32 ScriptOffset)
 {
 	TWeakObjectPtr<const UEdGraphNode>& Result = ScriptOffsetToNodes.FindOrAdd(ScriptOffset);
 	if (!Result.IsValid() && BlueprintClass.IsValid())
 	{
-		if (const UEdGraphNode* GraphNode = BlueprintClass->GetDebugData().FindSourceNodeFromCodeLocation(Function.Get(), ScriptOffset, true))
+		// First pass lookup, if inside a tunnel this will return the tunnel instance node.
+		Result = BlueprintClass->GetDebugData().FindSourceNodeFromCodeLocation(Function.Get(), ScriptOffset, true);
+		if (const UEdGraphNode* PotentialTunnelNode = Result.Get())
 		{
-			if (GraphNode->IsA<UK2Node_MacroInstance>())
+			// Check for tunnel nodes
+			if (FBlueprintEditorUtils::IsTunnelInstanceNode(PotentialTunnelNode))
 			{
-				GraphNode = BlueprintClass->GetDebugData().FindMacroSourceNodeFromCodeLocation(Function.Get(), ScriptOffset);
+				// Find the true source node.
+				if (const UEdGraphNode* TrueGraphNode = BlueprintClass->GetDebugData().FindMacroSourceNodeFromCodeLocation(Function.Get(), ScriptOffset))
+				{
+					// Cache the tunnel node.
+					if (TrueGraphNode != PotentialTunnelNode || FBlueprintEditorUtils::IsTunnelInstanceNode(TrueGraphNode))
+					{
+						NodeToTunnelNode.Add(TrueGraphNode) = PotentialTunnelNode;
+						Result = TrueGraphNode;
+					}
+				}
 			}
-			Result = GraphNode;
 		}
 	}
 	return Result.Get();
 }
 
-const UEdGraphNode* FBlueprintFunctionContext::FindTunnelSourceNodeFromCodeLocation(const int32 ScriptOffset)
+const UEdGraphNode* FBlueprintFunctionContext::GetTunnelNodeFromCodeLocation(const int32 ScriptOffset)
 {
 	TWeakObjectPtr<const UEdGraphNode> Result;
-	if (!Result.IsValid() && BlueprintClass.IsValid())
+	if (TWeakObjectPtr<const UEdGraphNode>* NodeSearch = ScriptOffsetToNodes.Find(ScriptOffset))
 	{
-		if (const UEdGraphNode* GraphNode = BlueprintClass->GetDebugData().FindMacroSourceNodeFromCodeLocation(Function.Get(), ScriptOffset))
+		if (TWeakObjectPtr<const UEdGraphNode>* TunnelNodeSearch = NodeToTunnelNode.Find(*NodeSearch))
 		{
-			Result = GraphNode;
+			Result = *TunnelNodeSearch;
 		}
 	}
 	return Result.Get();
 }
 
+const UEdGraphNode* FBlueprintFunctionContext::GetTunnelNodeFromGraphNode(const UEdGraphNode* GraphNode) const
+{
+	TWeakObjectPtr<const UEdGraphNode> Result;
+	if (const TWeakObjectPtr<const UEdGraphNode>* SearchResult = NodeToTunnelNode.Find(GraphNode))
+	{
+		Result = *SearchResult;
+	}
+	return Result.Get();
+}
+
+bool FBlueprintFunctionContext::IsNodeFromTunnelGraph(const UEdGraphNode* Node) const
+{
+	return NodeToTunnelNode.Contains(Node);
+}
+
 const UEdGraphPin* FBlueprintFunctionContext::GetPinFromCodeLocation(const int32 ScriptOffset)
 {
 	FEdGraphPinReference& Result = ScriptOffsetToPins.FindOrAdd(ScriptOffset);
@@ -1249,7 +1643,15 @@ TSharedPtr<FScriptExecutionNode> FBlueprintFunctionContext::CreateExecutionNode(
 	check(!ScriptExecNode.IsValid());
 	UEdGraph* OuterGraph = InitParams.ObservedObject.IsValid() ? InitParams.ObservedObject.Get()->GetTypedOuter<UEdGraph>() : nullptr;
 	InitParams.OwningGraphName = OuterGraph ? OuterGraph->GetFName() : NAME_None;
-	ScriptExecNode = MakeShareable(new FScriptExecutionNode(InitParams));
+	const bool bPureChain = (InitParams.NodeFlags & EScriptExecutionNodeFlags::PureChain) != 0U;
+	if (bPureChain)
+	{
+		ScriptExecNode = MakeShareable(new FScriptExecutionPureNode(InitParams));
+	}
+	else
+	{
+		ScriptExecNode = MakeShareable(new FScriptExecutionNode(InitParams));
+	}
 	return ScriptExecNode;
 }
 
@@ -1267,16 +1669,13 @@ template<typename ExecNodeType> TSharedPtr<ExecNodeType> FBlueprintFunctionConte
 void FBlueprintFunctionContext::MapTunnelInstance(UK2Node_Tunnel* TunnelInstance)
 {
 	// Grab tunnel context from instance name
-	const FName TunnelGraphName = TunnelInstance->GetFName();
-	TSharedPtr<FBlueprintTunnelInstanceContext> TunnelGraphContext = StaticCastSharedPtr<FBlueprintTunnelInstanceContext>(BlueprintContext.Pin()->GetFunctionContext(TunnelGraphName));
-	if (TunnelGraphContext.IsValid())
+	if (TunnelInstance)
 	{
-		// Add the instance context
-		ChildFunctionContexts.Add(TunnelGraphName) = TunnelGraphContext;
-		// Add any child contexts belonging to the instance.
-		for (auto ChildContexts : TunnelGraphContext->ChildFunctionContexts)
+		const FName ScopedFunctionName = GetTunnelInstanceFunctionName(TunnelInstance);
+		TSharedPtr<FBlueprintFunctionContext> TunnelFunctionContext = BlueprintContext.Pin()->GetFunctionContext(ScopedFunctionName);
+		if (TunnelFunctionContext.IsValid())
 		{
-			ChildFunctionContexts.Add(ChildContexts.Key) = ChildContexts.Value;
+			ChildFunctionContexts.Add(ScopedFunctionName) = TunnelFunctionContext;
 		}
 	}
 }
@@ -1291,6 +1690,11 @@ void FBlueprintFunctionContext::AddExitPoint(TSharedPtr<FScriptExecutionNode> Ex
 	ExitPoints.Add(ExitPoint);
 }
 
+FName FBlueprintFunctionContext::GetScopedEventName(const FName EventName) const
+{
+	return FName(*FString::Printf(TEXT("%s::%s"), *BlueprintClass.Get()->GetName(), *EventName.ToString()));
+}
+
 void FBlueprintFunctionContext::AddCallSiteEntryPointsToNode(TSharedPtr<FScriptExecutionNode> CallingNode) const
 {
 	if (CallingNode.IsValid())
@@ -1324,13 +1728,12 @@ void FBlueprintTunnelInstanceContext::MapTunnelContext(TSharedPtr<FBlueprintFunc
 		BlueprintClass = CallingFunctionContext->GetBlueprintClass();
 		TunnelInstanceNode = TunnelInstance;
 		// Find the function context that represents the graph and not the instance
-		UEdGraph* TunnelGraph = GetTunnelGraphFromNode(TunnelInstance);
-		TunnelFunctionContext = BlueprintContext.Pin()->GetFunctionContext(TunnelGraph->GetFName());
-		check (TunnelFunctionContext.IsValid());
+		UEdGraph* TunnelGraph = GetGraphFromNode(TunnelInstance, false);
 		// Map tunnel Input/Output
 		MapTunnelIO();
 		// Create tunnel instance node.
 		FScriptExecNodeParams TunnelInstanceParams;
+		TunnelInstanceParams.SampleFrequency = 1;
 		TunnelInstanceParams.NodeName = TunnelInstance->GetFName();
 		TunnelInstanceParams.ObservedObject = TunnelInstance;
 		TunnelInstanceParams.DisplayName = TunnelInstance->GetNodeTitle(ENodeTitleType::ListView);
@@ -1349,7 +1752,6 @@ void FBlueprintTunnelInstanceContext::MapTunnelContext(TSharedPtr<FBlueprintFunc
 			TunnelInstanceParams.Icon = const_cast<FSlateBrush*>(TunnelIcon);
 		}
 		TSharedPtr<FScriptExecutionTunnelInstance> TunnelInstanceExecNode = CreateTypedExecutionNode<FScriptExecutionTunnelInstance>(TunnelInstanceParams);
-		ExternalNodes.Add(TunnelInstance->GetFName()) = TunnelInstanceExecNode;
 		// Map child tunnel instances now, because we require that the instance exec node is setup.
 		TMap<UK2Node_Tunnel*, TSharedPtr<FBlueprintFunctionContext>> ChildTunnels;
 		DiscoverTunnels(TunnelGraph, ChildTunnels);
@@ -1357,8 +1759,6 @@ void FBlueprintTunnelInstanceContext::MapTunnelContext(TSharedPtr<FBlueprintFunc
 		{
 			MapTunnelInstance(ChildTunnel.Key);
 		}
-		// Add Tunnel Graph Context as child context
-		ChildFunctionContexts.Add(TunnelFunctionContext->GetFunctionName()) = TunnelFunctionContext;
 		// Find tunnel instance entry sites
 		TMap<FName, UEdGraphPin*> InstanceEntrySites;
 		for (auto Pin : TunnelInstance->Pins)
@@ -1385,11 +1785,12 @@ void FBlueprintTunnelInstanceContext::MapTunnelContext(TSharedPtr<FBlueprintFunc
 				// Create custom exec node for each tunnel instance entry, this is so we can customise the appearance of the node.
 				FName BoundaryName = GetTunnelBoundaryName(InstanceEntryPoint.Value);
 				FScriptExecNodeParams TunnelEntryParams;
+				TunnelEntryParams.SampleFrequency = 1;
 				TunnelEntryParams.NodeName = BoundaryName;
 				TunnelEntryParams.ObservedObject = TunnelInstance;
 				TunnelEntryParams.ObservedPin = InstanceEntryPoint.Value;
 				TunnelEntryParams.DisplayName = TunnelInstance->GetNodeTitle(ENodeTitleType::ListView);
-				TunnelEntryParams.Tooltip = LOCTEXT("NavigateToNodeLocationHyperlink_ToolTip", "Navigate to the Node");
+				TunnelEntryParams.Tooltip = LOCTEXT("NavigateToNodeLocationHyperlink_ToolTip", "Navigate to the Tunnel Instance");
 				TunnelEntryParams.NodeFlags = EScriptExecutionNodeFlags::TunnelEntryPinInstance;
 				TunnelEntryParams.IconColor = FLinearColor(1.f, 1.f, 1.f, 0.8f);
 				TunnelEntryParams.Icon = TunnelInstanceParams.Icon;
@@ -1398,7 +1799,6 @@ void FBlueprintTunnelInstanceContext::MapTunnelContext(TSharedPtr<FBlueprintFunc
 				InstanceTunnelEntry->AddChildNode(EntryPoint);
 				// Register the external entry point
 				TunnelInstanceExecNode->AddEntrySite(EntryPointScriptOffset, InstanceTunnelEntry);
-				ExternalNodes.Add(BoundaryName) = InstanceTunnelEntry;
 				// Set the current entry point we are mapping.
 				StagingEntryPoint = InstanceTunnelEntry;
 				// Map this tunnel entry
@@ -1408,34 +1808,29 @@ void FBlueprintTunnelInstanceContext::MapTunnelContext(TSharedPtr<FBlueprintFunc
 					{
 						// Pass through any non-relevant (e.g. reroute) nodes.
 						LinkedPin = FBlueprintEditorUtils::FindFirstCompilerRelevantLinkedPin(LinkedPin);
-						const UEdGraphNode* LinkedNode = LinkedPin->GetOwningNode();
-						const UK2Node_Tunnel* LinkedTunnelNode = Cast<UK2Node_Tunnel>(LinkedNode);
-						const UEdGraph* ChildTunnelGraph = GetTunnelGraphFromNode(LinkedTunnelNode);
-						const FName ChildGraphName = ChildTunnelGraph ? ChildTunnelGraph->GetFName() : NAME_None;
-						const bool bLinkedToTunnelExit = LinkedTunnelNode && ChildGraphName == GraphName;
-						const int32 ScriptCodeOffset = GetCodeLocationFromPin(LinkedPin);
-						if (bLinkedToTunnelExit)
+						// Grab the owning node
+						UEdGraphNode* LinkedNode = LinkedPin->GetOwningNode();
+						// Check for linkage to another tunnel, this could be another tunnel or an exit site for this one.
+						bool bLinkedToTunnel = false;
+						bool bInternalTunnel = false;
+						if (const UK2Node_Tunnel* LinkedTunnelNode = Cast<UK2Node_Tunnel>(LinkedNode))
 						{
-							// This handles tunnel input linked directly to tunnel output as a special case.
-							TSharedPtr<FScriptExecutionTunnelExit> OutputPinNode = GetTypedProfilerDataForNode<FScriptExecutionTunnelExit>(GetPinName(LinkedPin));
-							check(OutputPinNode.IsValid());
-							StagingEntryPoint->AddExitSite(ScriptCodeOffset, OutputPinNode);
-							EntryPoint->AddLinkedNode(ScriptCodeOffset, OutputPinNode);
+							// Marked as linked to another tunnel.
+							bLinkedToTunnel = true;
+							// Check if the tunnel is a tunnel instance node or not.
+							bInternalTunnel = !FBlueprintEditorUtils::IsTunnelInstanceNode(LinkedTunnelNode);
 						}
-						else
+						// Map the linked pin/node whilst being aware of directly linked tunnel boundaries.
+						TSharedPtr<FScriptExecutionNode> ExecNode = bLinkedToTunnel ? MapTunnelBoundary(LinkedPin) : MapNodeExecution(LinkedNode);
+						// Only add the new node if it is not an internal tunnel.
+						if (!bInternalTunnel && ExecNode.IsValid())
 						{
-							// This adds any encountered exit points to the active tunnel name.
-							const bool bLinkedToTunnel = ChildGraphName != NAME_None && ChildGraphName != GraphName;
-							TSharedPtr<FScriptExecutionNode> ExecNode = bLinkedToTunnel ? MapTunnelBoundary(LinkedPin) : MapNodeExecution(LinkedPin->GetOwningNode());
-							if (ExecNode.IsValid())
-							{
-								EntryPoint->AddLinkedNode(ScriptCodeOffset, ExecNode);
-							}
+							EntryPoint->AddChildNode(ExecNode);
 						}
 					}
 				}
 				// Add exit sites to instance
-				for (auto ExitSite : InstanceTunnelEntry->GetExitSites())
+				for (auto ExitSite : InstanceTunnelEntry->GetLinkedNodes())
 				{
 					if (ExitSite.Value->IsTunnelExit())
 					{
@@ -1446,12 +1841,19 @@ void FBlueprintTunnelInstanceContext::MapTunnelContext(TSharedPtr<FBlueprintFunc
 				StagingEntryPoint.Reset();
 			}
 		}
-		// Reset the mapping context
-		Function.Reset();
-		BlueprintClass.Reset();
 	}
 }
 
+TSharedPtr<FScriptExecutionNode> FBlueprintTunnelInstanceContext::GetPureNodeFromTunnelPin(const UEdGraphPin* TunnelPin) const
+{
+	TSharedPtr<FScriptExecutionNode> Result;
+	if (const TSharedPtr<FScriptExecutionNode>* SearchResult = PureLinkNodes.Find(TunnelPin))
+	{
+		Result = *SearchResult;
+	}
+	return Result;
+}
+
 TSharedPtr<FScriptExecutionNode> FBlueprintTunnelInstanceContext::MapNodeExecution(UEdGraphNode* NodeToMap)
 {
 	TSharedPtr<FScriptExecutionNode> MappedNode;
@@ -1472,16 +1874,101 @@ TSharedPtr<FScriptExecutionNode> FBlueprintTunnelInstanceContext::MapNodeExecuti
 	return MappedNode;
 }
 
+void FBlueprintTunnelInstanceContext::MapInputPins(TSharedPtr<FScriptExecutionNode> ExecNode, const TArray<UEdGraphPin*>& Pins)
+{
+	TArray<int32> PinScriptCodeOffsets;
+	TSharedPtr<FScriptExecutionNode> PureChainRootNode = ExecNode;
+
+	for (auto InputPin : Pins)
+	{
+		// If this input pin is linked to a pure node in the source graph, create and map all known execution paths for it.
+		for (auto LinkedPin : InputPin->LinkedTo)
+		{
+			// Pass through non-relevant (e.g. reroute) nodes.
+			UK2Node* OwningNode = FBlueprintEditorUtils::FindFirstCompilerRelevantNode(LinkedPin);
+
+			// If this is a tunnel node - we need to map through the tunnel
+			if (OwningNode && OwningNode->IsA<UK2Node_Tunnel>())
+			{
+				TSharedPtr<FScriptExecutionNode> TunnelPureNode;
+
+				if (FBlueprintEditorUtils::IsTunnelInstanceNode(OwningNode))
+				{
+					// Find the tunnel instance context
+					const FName ScopedTunnelContextName = GetTunnelInstanceFunctionName(OwningNode);
+					TSharedPtr<FBlueprintTunnelInstanceContext> TunnelContext = StaticCastSharedPtr<FBlueprintTunnelInstanceContext>(BlueprintContext.Pin()->GetFunctionContext(ScopedTunnelContextName));
+					check (TunnelContext.IsValid());
+					TunnelPureNode = TunnelContext->GetPureNodeFromTunnelPin(LinkedPin);
+				}
+				else
+				{
+					if (const UEdGraphPin* ExternalPin = FindMatchingPin(TunnelInstanceNode.Get(), LinkedPin))
+					{
+						if (TSharedPtr<FScriptExecutionNode>* SearchResult = PureLinkNodes.Find(ExternalPin))
+						{
+							TunnelPureNode = *SearchResult;
+						}
+					}
+				}
+				// Link up the pure nodes
+				if (TunnelPureNode.IsValid() && TunnelPureNode->IsPureChain())
+				{
+					// Grab the pure chain root.
+					PureChainRootNode = FindOrCreatePureChainRoot(ExecNode);
+					if (!PureChainRootNode.IsValid())
+					{
+						PureChainRootNode = ExecNode;
+					}
+					// Copy pure chain links
+					for (auto LinkedPureNode : TunnelPureNode->GetLinkedNodes())
+					{
+						ExecNode->AddLinkedNode(LinkedPureNode.Key, LinkedPureNode.Value);
+					}
+				}
+				continue;
+			}
+
+			// Note: Intermediate pure nodes can have output pins that masquerade as impure node output pins when links are "moved" from the source graph (thus
+			// resulting in a false association here with one or more script code offsets), so we must first ensure that the link is really to a pure node output.
+			if (OwningNode && OwningNode->IsNodePure())
+			{
+				GetAllCodeLocationsFromPin(LinkedPin, PinScriptCodeOffsets);
+				if (PinScriptCodeOffsets.Num() > 0)
+				{
+					// Add a pure chain container node as the root, if it's not already in place.
+					if (PureChainRootNode == ExecNode && !ExecNode->IsPureNode())
+					{
+						PureChainRootNode = FindOrCreatePureChainRoot(ExecNode);
+					}
+
+					TSharedPtr<FScriptExecutionNode> PureNode = MapNodeExecution(OwningNode);
+					check(PureNode.IsValid() && PureNode->IsPureNode());
+					for (int32 i = 0; i < PinScriptCodeOffsets.Num(); ++i)
+					{
+						PureChainRootNode->AddLinkedNode(PinScriptCodeOffsets[i], PureNode);
+					}
+				}
+			}
+		}
+	}
+}
+
+bool FBlueprintTunnelInstanceContext::IsTunnelNodeInternal(const UEdGraphNode* TunnelNode)
+{
+	if (!FBlueprintEditorUtils::IsTunnelInstanceNode(TunnelNode))
+	{
+		return TunnelNode->GetTypedOuter<UEdGraph>()->GetFName() == GraphName;
+	}
+	return false;
+}
+
 bool FBlueprintTunnelInstanceContext::IsPinFromThisTunnel(const UEdGraphPin* TunnelPin) const
 {
 	bool bIsFromThisTunnel = false;
-	UEdGraph* PinGraph = nullptr;
 	const UEdGraphNode* PinNode = TunnelPin ? TunnelPin->GetOwningNode() : nullptr;
 	if (PinNode)
 	{
-		// if this is a tunnel instance locate the tunnel graph, otherwise just use the outer graph.
-		PinGraph = PinNode->IsA<UK2Node_Tunnel>() ? GetTunnelGraphFromNode(PinNode) : PinNode->GetTypedOuter<UEdGraph>();
-		if (PinGraph)
+		if (UEdGraph* PinGraph = GetGraphFromNode(PinNode))
 		{
 			bIsFromThisTunnel = GraphName == PinGraph->GetFName();
 		}
@@ -1500,7 +1987,7 @@ void FBlueprintTunnelInstanceContext::DiscoverExitSites(TSharedPtr<FScriptExecut
 			if (UEdGraphPin* ExternalPin = ExitNode->GetExternalPin())
 			{
 				const int32 ScriptCodeOffset = GetCodeLocationFromPin(ExternalPin);
-				StagingEntryPoint->AddExitSite(ScriptCodeOffset, ExitNode);
+				StagingEntryPoint->AddLinkedNode(ScriptCodeOffset, ExitNode);
 			}
 		}
 	}
@@ -1519,11 +2006,7 @@ void FBlueprintTunnelInstanceContext::DiscoverExitSites(TSharedPtr<FScriptExecut
 
 TSharedPtr<FScriptExecutionNode> FBlueprintTunnelInstanceContext::MapTunnelBoundary(const UEdGraphPin* TunnelPin)
 {
-	TSharedPtr<FScriptExecutionNode> TunnelBoundaryNode;
-	// If this isn't an internal pin, this is a boundary to a different tunnel.
-	const bool bInternalTunnel = IsTunnelInternal(TunnelPin->GetOwningNode());
-	const FName TunnelBoundaryName = bInternalTunnel ? GetPinName(TunnelPin) : GetTunnelBoundaryName(TunnelPin);
-	TunnelBoundaryNode = GetProfilerDataForNodeChecked(TunnelBoundaryName);
+	TSharedPtr<FScriptExecutionNode> TunnelBoundaryNode = GetTunnelBoundaryNodeChecked(TunnelPin);
 	if (TunnelBoundaryNode->IsTunnelExit())
 	{
 		// Map tunnel exit because we are mapping inside a tunnel discovering exit pins.
@@ -1531,12 +2014,39 @@ TSharedPtr<FScriptExecutionNode> FBlueprintTunnelInstanceContext::MapTunnelBound
 		if (UEdGraphPin* ExternalPin = ExitNode->GetExternalPin())
 		{
 			const int32 ScriptCodeOffset = GetCodeLocationFromPin(ExternalPin);
-			StagingEntryPoint->AddExitSite(ScriptCodeOffset, ExitNode);
+			StagingEntryPoint->AddLinkedNode(ScriptCodeOffset, ExitNode);
 		}
 	}
 	else if (TunnelBoundaryNode->IsTunnelEntry())
 	{
-		return FBlueprintFunctionContext::MapTunnelBoundary(TunnelPin);
+		TSharedPtr<FScriptExecutionTunnelEntry> TunnelEntryInstance = StaticCastSharedPtr<FScriptExecutionTunnelEntry>(TunnelBoundaryNode);
+		for (auto ExitSite : TunnelEntryInstance->GetLinkedNodes())
+		{
+			if (ExitSite.Value->IsTunnelExit())
+			{
+				TSharedPtr<FScriptExecutionTunnelExit> TunnelExit = TunnelEntryInstance->GetExitSite(ExitSite.Key);
+				const UEdGraphPin* TunnelInstanceExitPin = TunnelExit->GetExternalPin();
+				for (auto LinkedPin : TunnelInstanceExitPin->LinkedTo)
+				{
+					UEdGraphNode* LinkedNode = LinkedPin->GetOwningNode();
+					TSharedPtr<FScriptExecutionNode> LinkedExecNode;
+					// Need to be careful here because a tunnel instance exit site can link to a tunnel Boundary too.
+					if (LinkedNode->IsA<UK2Node_Tunnel>())
+					{
+						LinkedExecNode = MapTunnelBoundary(LinkedPin);
+					}
+					else
+					{
+						LinkedExecNode = GetProfilerDataForNode(LinkedNode->GetFName());
+					}
+					if (!LinkedExecNode.IsValid())
+					{
+						LinkedExecNode = MapNodeExecution(LinkedNode);
+					}
+					TunnelExit->AddLinkedNode(ExitSite.Key, LinkedExecNode);
+				}
+			}
+		}
 	}
 	return TunnelBoundaryNode;
 }
@@ -1548,28 +2058,40 @@ void FBlueprintTunnelInstanceContext::MapTunnelIO()
 		// Map internal tunnel pins to tunnel instance pins
 		TMap<UEdGraphPin*, UEdGraphPin*> EntryPins;
 		TMap<UEdGraphPin*, UEdGraphPin*> ExitPins;
+		TMap<UEdGraphPin*, UEdGraphPin*> PurePins;
 		TArray<UK2Node_Tunnel*> GraphTunnels;
-		UEdGraph* TunnelGraph = GetTunnelGraphFromNode(TunnelInstanceGraphNode);
+		UEdGraph* TunnelGraph = GetGraphFromNode(TunnelInstanceGraphNode, false);
 		TunnelGraph->GetNodesOfClass<UK2Node_Tunnel>(GraphTunnels);
+		// Build tunnel pin sets.
 		for (auto Tunnel : GraphTunnels)
 		{
-			if (IsTunnelInternal(Tunnel))
+			if (IsTunnelNodeInternal(Tunnel))
 			{
 				for (UEdGraphPin* InternalPin : Tunnel->Pins)
 				{
-					if (InternalPin->LinkedTo.Num() && InternalPin->PinType.PinCategory == UEdGraphSchema_K2::PC_Exec)
+					UEdGraphPin* TunnelInstancePin = FindMatchingPin(TunnelInstanceGraphNode, InternalPin);
+					if (InternalPin->LinkedTo.Num())
 					{
-						if (UEdGraphPin* TunnelInstancePin = FindMatchingPin(TunnelInstanceGraphNode, InternalPin))
+						if (InternalPin->PinType.PinCategory == UEdGraphSchema_K2::PC_Exec)
 						{
-							if (InternalPin->Direction == EGPD_Output)
-							{
-								EntryPins.Add(InternalPin) = TunnelInstancePin;
-							}
-							else
+							// Map linked exec pins
+							if (TunnelInstancePin)
 							{
-								ExitPins.Add(InternalPin) = TunnelInstancePin;
+								if (InternalPin->Direction == EGPD_Output)
+								{
+									EntryPins.Add(InternalPin) = TunnelInstancePin;
+								}
+								else
+								{
+									ExitPins.Add(InternalPin) = TunnelInstancePin;
+								}
 							}
 						}
+						else
+						{
+							// Map pure pin sites
+							PurePins.Add(InternalPin) = TunnelInstancePin;
+						}
 					}
 				}
 			}
@@ -1578,6 +2100,7 @@ void FBlueprintTunnelInstanceContext::MapTunnelIO()
 		for (auto EntryPinSet : EntryPins)
 		{
 			FScriptExecNodeParams LinkParams;
+			LinkParams.SampleFrequency = 1;
 			LinkParams.NodeName = GetPinName(EntryPinSet.Key);
 			LinkParams.ObservedObject = EntryPinSet.Key->GetOwningNode();
 			LinkParams.ObservedPin = EntryPinSet.Key;
@@ -1596,6 +2119,7 @@ void FBlueprintTunnelInstanceContext::MapTunnelIO()
 		for (auto ExitPinSet : ExitPins)
 		{
 			FScriptExecNodeParams LinkParams;
+			LinkParams.SampleFrequency = 1;
 			LinkParams.NodeName = GetPinName(ExitPinSet.Key);
 			LinkParams.ObservedObject = ExitPinSet.Key->GetOwningNode();
 			LinkParams.ObservedPin = ExitPinSet.Key;
@@ -1610,17 +2134,58 @@ void FBlueprintTunnelInstanceContext::MapTunnelIO()
 			TSharedPtr<FScriptExecutionTunnelExit> ExitPoint = CreateTypedExecutionNode<FScriptExecutionTunnelExit>(LinkParams);
 			const FName TunnelBoundaryName = GetTunnelBoundaryName(ExitPinSet.Value);
 			// Add exit point under internal pin name for mapping.
-			ExternalNodes.Add(TunnelBoundaryName) = ExitPoint;
 			ExitPoint->SetExternalPin(ExitPinSet.Value);
 			AddExitPoint(ExitPoint);
 		}
+		// Map pure pin links
+		for (auto PurePinSet : PurePins)
+		{
+			// Use the external pin if tunnel entry otherwise use the internal pin.
+			const UEdGraphPin* PinToTrace = PurePinSet.Key->Direction == EGPD_Output ? PurePinSet.Value : PurePinSet.Key;
+			// Create a pure node chain for the pin
+			FScriptExecNodeParams PureChainParams;
+			static const FString PureChainNodeNameSuffix = TEXT("__PROFILER_InputPureTime");
+			FString PureChainNodeNameString = PinToTrace->GetName() + PureChainNodeNameSuffix;
+			PureChainParams.SampleFrequency = 1;
+			PureChainParams.NodeName = FName(*PureChainNodeNameString);
+			PureChainParams.ObservedPin = PinToTrace;
+			PureChainParams.DisplayName = LOCTEXT("PureChain_DisplayName", "Pure Time");
+			PureChainParams.Tooltip = LOCTEXT("PureChain_ToolTip", "Expand pure node timing");
+			PureChainParams.NodeFlags = EScriptExecutionNodeFlags::PureChain;
+			GetNodeCustomizations(PureChainParams);
+			const FSlateBrush* Icon = FEditorStyle::GetBrush(TEXT("BlueprintProfiler.PureNode"));
+			PureChainParams.Icon = const_cast<FSlateBrush*>(Icon);
+			TSharedPtr<FScriptExecutionNode> PureChainRootNode = CreateExecutionNode(PureChainParams);
+			PureLinkNodes.Add(PurePinSet.Value) = PureChainRootNode;
+			// Map pins
+			for (auto LinkedPin : PinToTrace->LinkedTo)
+			{
+				// Pass through non-relevant (e.g. reroute) nodes.
+				UK2Node* OwningNode = FBlueprintEditorUtils::FindFirstCompilerRelevantNode(LinkedPin);
+				// Map the execution
+				TSharedPtr<FScriptExecutionNode> PureNode = MapNodeExecution(OwningNode);
+				if (PureNode.IsValid() && PureNode->IsPureNode())
+				{
+					TArray<int32> CodeLocations;
+					GetAllCodeLocationsFromPin(LinkedPin, CodeLocations);
+					for (int32 CodeLocation : CodeLocations)
+					{
+						PureChainRootNode->AddLinkedNode(CodeLocation, PureNode);
+					}
+				}
+			}
+		}
 	}
 }
 
-TSharedPtr<FScriptExecutionNode> FBlueprintTunnelInstanceContext::CreateExecutionNode(FScriptExecNodeParams& InitParams)
+bool FBlueprintTunnelInstanceContext::GetProfilerContextFromScriptOffset(const int32 ScriptOffset, TSharedPtr<FScriptExecutionNode>& ExecNode, TSharedPtr<FBlueprintFunctionContext>& FunctionContext)
 {
-	// Use the function context that represents the tunnel graph.
-	return TunnelFunctionContext->CreateExecutionNode(InitParams);
+	if (const UEdGraphNode* GraphNode = GetNodeFromCodeLocation(ScriptOffset))
+	{
+		FunctionContext = AsShared();
+		ExecNode = GetProfilerDataForNode(GraphNode->GetFName());
+	}
+	return ExecNode.IsValid() && FunctionContext.IsValid();
 }
 
 //////////////////////////////////////////////////////////////////////////
@@ -1636,9 +2201,9 @@ bool FScriptEventPlayback::Process(const TArray<FScriptInstrumentedEvent>& Signa
 	if (BlueprintContext.IsValid() && InstanceName != NAME_None)
 	{
 		check(SignalData[StartIdx].IsEvent());
-		EventName = bEventIsResuming ? EventName : SignalData[StartIdx].GetFunctionName();
-		CurrentFunctionName = BlueprintContext->GetEventFunctionName(SignalData[StartIdx].GetFunctionName());
-		TSharedPtr<FBlueprintFunctionContext> FunctionContext = BlueprintContext->GetFunctionContext(CurrentFunctionName);
+		EventName = bEventIsResuming ? EventName : SignalData[StartIdx].GetScopedFunctionName();
+		TSharedPtr<FBlueprintFunctionContext> FunctionContext = BlueprintContext->GetFunctionContextForEventChecked(EventName);
+		CurrentFunctionName = FunctionContext->GetFunctionName();
 		TSharedPtr<FScriptExecutionNode> EventNode = FunctionContext->GetProfilerDataForNode(EventName);
 		// Find Associated graph nodes and submit into a node map for later processing.
 		TMap<const UEdGraphNode*, NodeSignalHelper> NodeMap;
@@ -1646,6 +2211,16 @@ bool FScriptEventPlayback::Process(const TArray<FScriptInstrumentedEvent>& Signa
 		int32 LastEventIdx = StartIdx;
 		const int32 EventStartOffset = SignalData[StartIdx].IsResumeEvent() ? 3 : 1;
 		LatentLinkId = bEventIsResuming ? INDEX_NONE : LatentLinkId;
+		// If we have sub calls into inherited events we don't have properly formed events, find the correct event now.
+		if (!EventNode.IsValid() && FunctionContext.IsValid())
+		{
+			const FScriptInstrumentedEvent& FirstValidSignal = SignalData[StartIdx + EventStartOffset];
+			if (const UK2Node_Event* EventGraphNode = Cast<UK2Node_Event>(FunctionContext->GetNodeFromCodeLocation(FirstValidSignal.GetScriptCodeOffset())))
+			{
+				EventName = FunctionContext->GetScopedEventName(EventGraphNode->GetFunctionName());
+				EventNode = FunctionContext->GetProfilerDataForNode(EventName);
+			}
+		}
 
 		for (int32 SignalIdx = StartIdx + EventStartOffset; SignalIdx < StopIdx; ++SignalIdx)
 		{
@@ -1654,7 +2229,7 @@ bool FScriptEventPlayback::Process(const TArray<FScriptInstrumentedEvent>& Signa
 			if (CurrSignal.GetType() == EScriptInstrumentation::Class)
 			{
 				// Update the current mapped blueprint context.
-				BlueprintContext = Profiler->GetBlueprintContext(CurrSignal.GetObjectPath());
+				BlueprintContext = Profiler->GetBlueprintContext(CurrSignal.GetBlueprintClassPath());
 
 				// Skip to the next signal.
 				continue;
@@ -1662,16 +2237,16 @@ bool FScriptEventPlayback::Process(const TArray<FScriptInstrumentedEvent>& Signa
 			else if (CurrSignal.GetType() == EScriptInstrumentation::Instance)
 			{
 				// Update the current mapped instance name.
-				InstanceName = BlueprintContext->MapBlueprintInstance(CurrSignal.GetObjectPath());
+				InstanceName = BlueprintContext->MapBlueprintInstance(CurrSignal.GetInstancePath());
 
 				// Skip to the next signal.
 				continue;
 			}
 
 			// Update script function.
-			if (CurrentFunctionName != CurrSignal.GetFunctionName())
+			if (CurrentFunctionName != CurrSignal.GetScopedFunctionName())
 			{
-				CurrentFunctionName = CurrSignal.GetFunctionName();
+				CurrentFunctionName = CurrSignal.GetScopedFunctionName();
 				FunctionContext = BlueprintContext->GetFunctionContext(CurrentFunctionName);
 				check(FunctionContext.IsValid());
 			}
@@ -1679,28 +2254,15 @@ bool FScriptEventPlayback::Process(const TArray<FScriptInstrumentedEvent>& Signa
 			{
 				NodeSignalHelper& CurrentNodeData = NodeMap.FindOrAdd(GraphNode);
 				// Initialise the current node context.
-				if (!CurrentNodeData.FunctionContext.IsValid())
+				if (!CurrentNodeData.IsValid())
 				{
-					CurrentNodeData.BlueprintContext = BlueprintContext;
-					CurrentNodeData.FunctionContext = FunctionContext;
-					CurrentNodeData.ImpureNode = FunctionContext->GetProfilerDataForGraphNode(GraphNode);
-					check(CurrentNodeData.ImpureNode.IsValid());
+					FunctionContext->GetProfilerContextFromScriptOffset(CurrSignal.GetScriptCodeOffset(), CurrentNodeData.ImpureNode, CurrentNodeData.FunctionContext);
+					check(CurrentNodeData.IsValid());
 				}
 				// Check for tunnel boundries and process here
 				if (CurrentNodeData.ImpureNode->IsTunnelInstance())
 				{
-					TWeakObjectPtr<const UEdGraphNode> NodeToUse = GraphNode;
-					if (const UEdGraphNode* RealSourceNode = FunctionContext->FindTunnelSourceNodeFromCodeLocation(CurrSignal.GetScriptCodeOffset()))
-					{
-						NodeToUse = RealSourceNode;
-					}
-					NodeSignalHelper& TunnelNodeData = NodeMap.FindOrAdd(NodeToUse.Get());
-					if (!TunnelNodeData.FunctionContext.IsValid())
-					{
-						TunnelNodeData.BlueprintContext = BlueprintContext;
-						TunnelNodeData.FunctionContext = FunctionContext;
-					}
-					ProcessTunnelBoundary(TunnelNodeData, CurrSignal);
+					ProcessTunnelBoundary(CurrentNodeData, CurrSignal);
 					continue;
 				}
 				// Process node data
@@ -1710,13 +2272,14 @@ bool FScriptEventPlayback::Process(const TArray<FScriptInstrumentedEvent>& Signa
 					{
 						if (const UEdGraphPin* Pin = FunctionContext->GetPinFromCodeLocation(CurrSignal.GetScriptCodeOffset()))
 						{
-							if (FunctionContext->HasProfilerDataForNode(Pin->GetOwningNode()->GetFName()))
+							// Use the impure nodes function context to look up the pure graph node because it may not be cached yet.
+							if (CurrentNodeData.FunctionContext->HasProfilerDataForGraphNode(Pin->GetOwningNode()))
 							{
 								TracePath.AddExitPin(CurrSignal.GetScriptCodeOffset());
 								CurrentNodeData.InputTracePaths.Insert(TracePath, 0);
 							}
 						}
-						CurrentNodeData.ExclusiveEvents.Add(CurrSignal);
+						CurrentNodeData.AverageEvents.Add(CurrSignal);
 						break;
 					}
 					case EScriptInstrumentation::NodeEntry:
@@ -1728,21 +2291,24 @@ bool FScriptEventPlayback::Process(const TArray<FScriptInstrumentedEvent>& Signa
 							// Ensure this is a different event
 							if (const UK2Node_Event* EventGraphNode = Cast<UK2Node_Event>(GraphNode))
 							{
-								const FName NewEventName = EventGraphNode->GetFunctionName();
+								const FName NewEventName = FunctionContext->GetScopedEventName(EventGraphNode->GetFunctionName());
 								if (NewEventName != EventName)
 								{
 									TracePath.Reset();
-									EventTimings.Add(EventName) = SignalData[SignalIdx].GetTime() - SignalData[LastEventIdx].GetTime();
+									if (EventNode.IsValid())
+									{
+										EventTimings.Add(EventName) = SignalData[SignalIdx].GetTime() - SignalData[LastEventIdx].GetTime();
+									}
 									EventName = NewEventName;
 									EventNode = FunctionContext->GetProfilerDataForNode(EventName);
 									LastEventIdx = SignalIdx;
 								}
 							}
 						}
-						// Process tunnel timings
-						CurrentNodeData.ExclusiveTracePaths.Push(TracePath);
-						CurrentNodeData.ExclusiveEvents.Add(CurrSignal);
-						AddToTraceHistory(CurrentNodeData.ImpureNode->GetName(), CurrSignal);
+						// Add node timings
+						CurrentNodeData.AverageTracePaths.Push(TracePath);
+						CurrentNodeData.AverageEvents.Add(CurrSignal);
+						AddToTraceHistory(CurrentNodeData.ImpureNode, CurrSignal);
 						break;
 					}
 					case EScriptInstrumentation::PushState:
@@ -1783,14 +2349,14 @@ bool FScriptEventPlayback::Process(const TArray<FScriptInstrumentedEvent>& Signa
 						{
 							TSharedPtr<FScriptExecutionTunnelInstance> TunnelInstance = TunnelEntry->GetTunnelInstance();
 							check (TunnelInstance.IsValid());
-							const UEdGraphNode* TunnelNode = Cast<UEdGraphNode>(TunnelInstance->GetObservedObject());
+							const UEdGraphNode* TunnelNode = TunnelInstance->GetTypedObservedObject<UEdGraphNode>();
 							NodeSignalHelper& TunnelNodeData = NodeMap.FindOrAdd(TunnelNode);
 							FScriptInstrumentedEvent OverrideEvent(CurrSignal);
 							FTracePath TunnelTrace;
 							TracePath.GetTunnelTracePath(TunnelTrace);
 							OverrideEvent.OverrideType(EScriptInstrumentation::NodeEntry);
-							TunnelNodeData.ExclusiveEvents.Add(OverrideEvent);
-							TunnelNodeData.ExclusiveTracePaths.Add(TunnelTrace);
+							TunnelNodeData.AverageEvents.Add(OverrideEvent);
+							TunnelNodeData.AverageTracePaths.Add(TunnelTrace);
 							TunnelTraceStack.Push(TunnelTrace);
 						}
 						break;
@@ -1825,23 +2391,23 @@ bool FScriptEventPlayback::Process(const TArray<FScriptInstrumentedEvent>& Signa
 					case EScriptInstrumentation::NodeExit:
 					{
 						// Cleanup branching multiple exits and correct the tracepath
-						if (CurrentNodeData.ExclusiveEvents.Num() && CurrentNodeData.ExclusiveEvents.Last().GetType() == EScriptInstrumentation::NodeExit)
+						if (CurrentNodeData.AverageEvents.Num() && CurrentNodeData.AverageEvents.Last().GetType() == EScriptInstrumentation::NodeExit)
 						{
-							CurrentNodeData.ExclusiveEvents.Pop();
-							if (CurrentNodeData.ExclusiveTracePaths.Num())
+							CurrentNodeData.AverageEvents.Pop();
+							if (CurrentNodeData.AverageTracePaths.Num())
 							{
-								TracePath = CurrentNodeData.ExclusiveTracePaths.Last();
+								TracePath = CurrentNodeData.AverageTracePaths.Last();
 							}
 						}
 						// Add Trace History
-						AddToTraceHistory(CurrentNodeData.ImpureNode->GetName(), CurrSignal);
+						AddToTraceHistory(CurrentNodeData.ImpureNode, CurrSignal);
 						// Process node exit
 						const int32 ScriptCodeExit = CurrSignal.GetScriptCodeOffset();
 						if (const UEdGraphPin* ValidPin = FunctionContext->GetPinFromCodeLocation(ScriptCodeExit))
 						{
 							TracePath.AddExitPin(ScriptCodeExit);
 						}
-						CurrentNodeData.ExclusiveEvents.Add(CurrSignal);
+						CurrentNodeData.AverageEvents.Add(CurrSignal);
 						// Process cyclic linkage - reset to root link tracepath
 						TSharedPtr<FScriptExecutionNode> NextLink = CurrentNodeData.ImpureNode->GetLinkedNodeByScriptOffset(ScriptCodeExit);
 						if (NextLink.IsValid() && NextLink->HasCyclicLinkage())
@@ -1849,9 +2415,9 @@ bool FScriptEventPlayback::Process(const TArray<FScriptInstrumentedEvent>& Signa
 							const UEdGraphNode* LinkedGraphNode = Cast<UEdGraphNode>(NextLink->GetObservedObject());
 							if (NodeSignalHelper* LinkedEntry = NodeMap.Find(LinkedGraphNode))
 							{
-								if (LinkedEntry->ExclusiveTracePaths.Num())
+								if (LinkedEntry->AverageTracePaths.Num())
 								{
-									TracePath = LinkedEntry->ExclusiveTracePaths.Last();
+									TracePath = LinkedEntry->AverageTracePaths.Last();
 								}
 							}
 						}
@@ -1859,33 +2425,31 @@ bool FScriptEventPlayback::Process(const TArray<FScriptInstrumentedEvent>& Signa
 					}
 					default:
 					{
-						CurrentNodeData.ExclusiveEvents.Add(CurrSignal);
+						CurrentNodeData.AverageEvents.Add(CurrSignal);
 						break;
 					}
 				}
 			}
 		}
 		// Process last event timing
-		double* TimingData = bEventIsResuming ? EventTimings.Find(EventName) : nullptr;
-		if (TimingData)
-		{
-			*TimingData += SignalData[StopIdx].GetTime() - SignalData[LastEventIdx].GetTime();
-		}
-		else
+		if (EventNode.IsValid())
 		{
-			EventTimings.Add(EventName) = SignalData[StopIdx].GetTime() - SignalData[LastEventIdx].GetTime();
+			double* TimingData = bEventIsResuming ? EventTimings.Find(EventName) : nullptr;
+			if (TimingData)
+			{
+				*TimingData += SignalData[StopIdx].GetTime() - SignalData[LastEventIdx].GetTime();
+			}
+			else
+			{
+				EventTimings.Add(EventName) = SignalData[StopIdx].GetTime() - SignalData[LastEventIdx].GetTime();
+			}
 		}
 		// Process outstanding event timings, adding to previous timings if existing.
 		for (auto EventTiming : EventTimings)
 		{
 			FTracePath EventTracePath;
 			// The UCS, along with BP events declared as 'const' in native C++ code, are implemented as standalone functions, and not within the ubergraph function context.
-			FunctionContext = BlueprintContext->GetFunctionContext(EventTiming.Key);
-			if (!FunctionContext.IsValid())
-			{
-				FunctionContext = BlueprintContext->GetFunctionContextForEvent(EventTiming.Key);
-			}
-			check (FunctionContext.IsValid());
+			FunctionContext = BlueprintContext->GetFunctionContextForEventChecked(EventTiming.Key);
 			EventNode = FunctionContext->GetProfilerDataForNode(EventTiming.Key);
 			check (EventNode.IsValid());
 			TSharedPtr<FScriptPerfData> PerfData = EventNode->GetOrAddPerfDataByInstanceAndTracePath(InstanceName, EventTracePath);
@@ -1897,13 +2461,15 @@ bool FScriptEventPlayback::Process(const TArray<FScriptInstrumentedEvent>& Signa
 		{
 			TSharedPtr<FScriptExecutionNode> ExecNode = CurrentNodeData.Value.ImpureNode;
 			TSharedPtr<FScriptExecutionNode> PureNode;
+			TSharedPtr<FBlueprintFunctionContext> NodeFunctionContext = CurrentNodeData.Value.FunctionContext;
+			check(NodeFunctionContext.IsValid());
 			double PureNodeEntryTime = 0.0;
 			double PureChainEntryTime = 0.0;
 			double NodeEntryTime = 0.0;
 			int32 ExclTracePathIdx = 0;
 			int32 InclTracePathIdx = 0;
 			// Process exclusive events for this node
-			for (auto EventIter = CurrentNodeData.Value.ExclusiveEvents.CreateIterator(); EventIter; ++EventIter)
+			for (auto EventIter = CurrentNodeData.Value.AverageEvents.CreateIterator(); EventIter; ++EventIter)
 			{
 				switch(EventIter->GetType())
 				{
@@ -1921,15 +2487,12 @@ bool FScriptEventPlayback::Process(const TArray<FScriptInstrumentedEvent>& Signa
 						PureNode.Reset();
 						PureNodeEntryTime = EventIter->GetTime();
 						
-						FunctionContext = CurrentNodeData.Value.BlueprintContext->GetFunctionContext(EventIter->GetFunctionName());
-						check(FunctionContext.IsValid());
-
 						const int32 ScriptCodeOffset = EventIter->GetScriptCodeOffset();
-						if (const UEdGraphPin* Pin = FunctionContext->GetPinFromCodeLocation(ScriptCodeOffset))
+						if (const UEdGraphPin* Pin = NodeFunctionContext->GetPinFromCodeLocation(ScriptCodeOffset))
 						{
-							if (FunctionContext->HasProfilerDataForNode(Pin->GetOwningNode()->GetFName()))
+							if (NodeFunctionContext->HasProfilerDataForNode(Pin->GetOwningNode()->GetFName()))
 							{
-								PureNode = CurrentNodeData.Value.FunctionContext->GetProfilerDataForNode(Pin->GetOwningNode()->GetFName());
+								PureNode = NodeFunctionContext->GetProfilerDataForNode(Pin->GetOwningNode()->GetFName());
 							}
 						}
 						break;
@@ -1954,12 +2517,12 @@ bool FScriptEventPlayback::Process(const TArray<FScriptInstrumentedEvent>& Signa
 					case EScriptInstrumentation::NodeExit:
 					{
 						check(NodeEntryTime != 0.0 );
-						const FTracePath NodeTracePath = CurrentNodeData.Value.ExclusiveTracePaths[ExclTracePathIdx++];
+						const FTracePath NodeTracePath = CurrentNodeData.Value.AverageTracePaths[ExclTracePathIdx++];
 						TSharedPtr<FScriptPerfData> PerfData = ExecNode->GetOrAddPerfDataByInstanceAndTracePath(InstanceName, NodeTracePath);
 						double PureChainDuration = PureChainEntryTime != 0.0 ? (NodeEntryTime - PureChainEntryTime) : 0.0;
 						double NodeDuration = EventIter->GetTime() - NodeEntryTime;
 						PerfData->AddEventTiming(NodeDuration);
-						PerfData->AddInclusiveTiming(PureChainDuration, false, false);
+						PerfData->AddInclusiveTiming(PureChainDuration, false);
 						TSharedPtr<FScriptExecutionNode> PureChainNode = ExecNode->GetPureChainNode();
 						if (PureChainNode.IsValid())
 						{
@@ -2006,7 +2569,7 @@ bool FScriptEventPlayback::Process(const TArray<FScriptInstrumentedEvent>& Signa
 							const FTracePath NodeTracePath = CurrentNodeData.Value.InclusiveTracePaths[InclTracePathIdx++];
 							TSharedPtr<FScriptPerfData> PerfData = ExecNode->GetOrAddPerfDataByInstanceAndTracePath(InstanceName, NodeTracePath);
 							double NodeDuration = EventIter.GetTime() - NodeEntryTime;
-							PerfData->AddInclusiveTiming(NodeDuration, false, true);
+							PerfData->AddInclusiveTiming(NodeDuration, false);
 							break;
 						}
 					}
@@ -2028,15 +2591,21 @@ void FScriptEventPlayback::ProcessTunnelBoundary(NodeSignalHelper& CurrentNodeDa
 	{
 		if (TunnelBoundary->HasFlags(EScriptExecutionNodeFlags::TunnelEntry))
 		{
+			// Grab the tunnel entry
+			TSharedPtr<FScriptExecutionTunnelEntry> TunnelEntry = StaticCastSharedPtr<FScriptExecutionTunnelEntry>(TunnelBoundary);
 			// Process tunnel entry sites
 			TunnelTraceStack.Push(TracePath);
-			TracePath = FTracePath(TracePath, StaticCastSharedPtr<FScriptExecutionTunnelEntry>(TunnelBoundary));
-			CurrentNodeData.ExclusiveEvents.Add(CurrSignal);
+			TracePath = FTracePath(TracePath, TunnelEntry);
+			CurrentNodeData.AverageEvents.Add(CurrSignal);
+			// Update the tunnel entry count so we can fixup stat samples later.
+			TunnelEntry->IncrementTunnelEntryCount();
+			// Add Trace History
+			AddToTraceHistory(TunnelBoundary, CurrSignal);
 		}
 		else if (TunnelBoundary->HasFlags(EScriptExecutionNodeFlags::TunnelExit))
 		{
 			// Process tunel exit sites.
-			if (CurrentNodeData.ExclusiveEvents.Num())
+			if (CurrentNodeData.AverageEvents.Num())
 			{
 				TSharedPtr<FScriptExecutionTunnelEntry> Tunnel = TracePath.GetTunnel();
 				FTracePath InternalExitTrace = TracePath;
@@ -2044,11 +2613,13 @@ void FScriptEventPlayback::ProcessTunnelBoundary(NodeSignalHelper& CurrentNodeDa
 				if (Tunnel.IsValid())
 				{
 					// Update the entry and both exit sites ( one inside the tunnel and one exit site )
-					const double TunnelTiming = CurrSignal.GetTime() - CurrentNodeData.ExclusiveEvents.Last().GetTime();
+					const double TunnelTiming = CurrSignal.GetTime() - CurrentNodeData.AverageEvents.Last().GetTime();
 					Tunnel->AddTunnelTiming(InstanceName, TracePath, InternalExitTrace, ScriptCodeOffset, TunnelTiming);
 				}
 				TracePath.AddExitPin(ScriptCodeOffset);
-				CurrentNodeData.ExclusiveEvents.Reset();
+				CurrentNodeData.AverageEvents.Reset();
+				// Add Trace History
+				AddToTraceHistory(TunnelBoundary, CurrSignal);
 			}
 		}
 	}
@@ -2058,29 +2629,29 @@ void FScriptEventPlayback::ProcessExecutionSequence(NodeSignalHelper& CurrentNod
 {
 	switch(CurrSignal.GetType())
 	{
-		// For a sequence node a restore state represents the end of execution of a sequence pin's path (excluding the last pin)
-		case EScriptInstrumentation::RestoreState:
+		// For a sequence node a push state represents the start of execution.
+		case EScriptInstrumentation::PushState:
 		{
-			// Convert the restore state into a node exit signal
+			// Convert the push state into a inclusive entry signal
 			FScriptInstrumentedEvent OverrideEvent(CurrSignal);
 			OverrideEvent.OverrideType(EScriptInstrumentation::NodeEntry);
-			CurrentNodeData.ExclusiveEvents.Add(OverrideEvent);
-			CurrentNodeData.ExclusiveTracePaths.Add(TracePath);
-			// Add Trace History
-			AddToTraceHistory(CurrentNodeData.ImpureNode->GetName(), OverrideEvent);
+			CurrentNodeData.InclusiveEvents.Add(OverrideEvent);
+			CurrentNodeData.InclusiveTracePaths.Add(TracePath);
 			break;
 		}
-		// For a sequence node a node entry represents the start of execution.
-		case EScriptInstrumentation::PushState:
+		// For a sequence node a restore state represents the end of execution of a sequence pin's path (excluding the last pin)
+		case EScriptInstrumentation::RestoreState:
 		{
-			// Convert the push state into a inclusive entry signal
+			// Convert the restore state into a node exit signal
 			FScriptInstrumentedEvent OverrideEvent(CurrSignal);
 			OverrideEvent.OverrideType(EScriptInstrumentation::NodeEntry);
-			CurrentNodeData.InclusiveEvents.Add(OverrideEvent);
-			CurrentNodeData.InclusiveTracePaths.Add(TracePath);
+			CurrentNodeData.AverageEvents.Add(OverrideEvent);
+			CurrentNodeData.AverageTracePaths.Add(TracePath);
+			// Add Trace History
+			AddToTraceHistory(CurrentNodeData.ImpureNode, OverrideEvent);
 			break;
 		}
-		// For a sequence node a restore state represents the end of execution of all the sequence pins.
+		// For a sequence node a pop state represents the end of execution of all the sequence pins.
 		case EScriptInstrumentation::PopState:
 		{
 			// Convert the pop state into a inclusive exit signal
@@ -2093,7 +2664,7 @@ void FScriptEventPlayback::ProcessExecutionSequence(NodeSignalHelper& CurrentNod
 	}
 }
 
-void FScriptEventPlayback::AddToTraceHistory(const FName NodeName, const FScriptInstrumentedEvent& TraceSignal)
+void FScriptEventPlayback::AddToTraceHistory(const TSharedPtr<FScriptExecutionNode> ProfilerNode, const FScriptInstrumentedEvent& TraceSignal)
 {
 	// Add Trace History
 	FBlueprintExecutionTrace& NewTraceEvent = BlueprintContext->AddNewTraceHistoryEvent();
@@ -2101,7 +2672,7 @@ void FScriptEventPlayback::AddToTraceHistory(const FName NodeName, const FScript
 	NewTraceEvent.TracePath = TracePath;
 	NewTraceEvent.InstanceName = InstanceName;
 	NewTraceEvent.FunctionName = CurrentFunctionName;
-	NewTraceEvent.NodeName = NodeName;
+	NewTraceEvent.ProfilerNode = ProfilerNode;
 	NewTraceEvent.Offset = TraceSignal.GetScriptCodeOffset();
 	NewTraceEvent.ObservationTime = TraceSignal.GetTime();
 }
