Copying //UE4/Dev-Blueprints to //UE4/Dev-Main (Source: //UE4/Dev-Blueprints @ 3130440)

==========================
MAJOR FEATURES + CHANGES
==========================

Change 3050029 on 2016/07/14 by [REDACTED_PERSON]

	This modifies the blueprint instrumented compilation chain so only the the blueprint you compile and all dependencies are instrumented and the profiler is notified rather than waiting for event data.
	#Proj BlueprintProfiler, Kismet, UnrelEd

	- This also improves the execution graph UI, notifying the user that no instances are available to display data from.

Change 3101549 on 2016/08/25 by [REDACTED_PERSON]

	BP nativization: fixed FEmitDefaultValueHelper::HandleInstancedSubobject

	https://udn.unrealengine.com/questions/308800/nativized-blueprints-newobject-call-uses-incorrect.html

Change 3101811 on 2016/08/25 by [REDACTED_PERSON]

	BP Profiler: Fixed stack overflow crash when compiling blueprints with nested macros

Change 3102478 on 2016/08/26 by [REDACTED_PERSON]

	BP Nativization:
	- improved native cast
	- improved bool handling

Change 3102944 on 2016/08/26 by [REDACTED_PERSON]

	[UE-33017] Don't include transient properties when generating property lists at cook time for optimized runtime Blueprint component instancing. Also ensure that deprecated properties are serialized during load/instancing at runtime.

	change summary:
	- modified FBlueprintComponentInstanceDataLoader to append 'PPF_UseDeprecatedProperties' to the FArchive port flags.
	- modified FBlueprintComponentInstanceDataWriter to append both 'PPF_Duplicate' and 'PPF_UseDeprecatedProperties" to the FArchive port flags (to ensure consistency w/ the instancing side).
	- switched the RecursivePropertyGatherLambda helper to a static class method instead
	- modified the RecursivePropertyGather utility method to exclude transient properties.

	notes:
	- the primary cause of UE-33017 was that UBodySetup can "share" the ShapeBodySetup object across all instances, but the shared object is not owned by the CDO, it's owned by the archetype. this caused the archetype to differ from the CDO, which caused us to emit the transient property at cook time. thsi threw off the serialization offset between read/write FArchive passes at runtime. since transient properties are not serialized as part of the template, there's no need to include them in the generated delta property list, so as a fix, i'm just excluding them altogether.

Change 3103692 on 2016/08/27 by [REDACTED_PERSON]

	Merging //UE4/Dev-Main to Dev-Blueprints (//UE4/Dev-Blueprints)

Change 3104266 on 2016/08/29 by [REDACTED_PERSON]

	Add test script to native assets for QAGame.

Change 3104399 on 2016/08/29 by [REDACTED_PERSON]

	Fix missing property warning in build script.

Change 3104419 on 2016/08/29 by [REDACTED_PERSON]

	- Reduced number of DynamicCLass instance dependencies
	- Fixed UDS default values dependencies
	- Improved WeakObjPtr handling
	- Improved const parameters handling

Change 3104474 on 2016/08/29 by [REDACTED_PERSON]

	BP Profiler: Fixed issue where collapsed nodes that share a name with a parent class collapsed node can cause a stack overflow

Change 3105605 on 2016/08/30 by [REDACTED_PERSON]

	Temp change: CIS Test

Change 3105738 on 2016/08/30 by [REDACTED_PERSON]

	UAT, CIS: testing NoRecompileUAT switch.

Change 3105800 on 2016/08/30 by [REDACTED_PERSON]

	UAT, CIS, Nativization:
	- reverted NoRecompileUAT switch.
	- testing nativization with -nocompileeditor flag and without -compile flag

Change 3106162 on 2016/08/30 by [REDACTED_PERSON]

	UAT, CIS, Nativization:
	-NoSubmit flag added. Otherwise UAT files are singed (when they are used by other process). It causes an error.
	- Ugly hack removed.

Change 3106261 on 2016/08/30 by [REDACTED_PERSON]

	[UE-34705] Gracefully handle tunnel node entry exec pins that aren't internally linked during BP profiler tunnel boundary mapping.

	change summary:
	- added FBlueprintFunctionContext::GetTunnelBoundaryNode() (uncheckedl variant).
	- moved FBlueprintFunctionContext::GetTunnelBoundaryNodeChecked() impl into GetTunnelBoundaryNode().
	- re-implemented FBlueprintFunctionContext::GetTunnelBoundaryNodeChecked() to call GetTunnelBoundaryNode() and then assert on the result.
	- changed the FBlueprintTunnelInstanceContext::GetTunnelBoundaryNodeChecked() impl to override GetTunnelBoundaryNode() instead.
	- modified FBlueprintFunctionContext::MapTunnelBoundary() to only process the entry case if the TunnelBoundaryNode result is valid. this way we simply skip tunnel boundary mapping if an entry path was not previously mapped (rather than assert).

Change 3106478 on 2016/08/30 by [REDACTED_PERSON]

	Include *.uasset files on builders running the NativizeAssets job.

Change 3107514 on 2016/08/31 by [REDACTED_PERSON]

	This set of changes is the result of a full pass on the blueprint profiler heat interface to try and bring them into a usable state.
	#Proj GraphEditor, Kismet, BlueprintProfiler,

Change 3108268 on 2016/08/31 by [REDACTED_PERSON]

	Minor change from profiler review sessions to move macro timing to average stats.
	#Proj Kismet

Change 3108991 on 2016/08/31 by [REDACTED_PERSON]

	UAT, CIS, Nativization: Test separate cooking and compiling

Change 3110097 on 2016/09/01 by [REDACTED_PERSON]

	Minor update to the blueprint profiler mapping functionality to ignore disabled nodes and a fix for the max timing white glow bug.
	#Proj Kismet, BlueprintProfiler

Change 3114553 on 2016/09/06 by [REDACTED_PERSON]

	Support for TMap/TSet in blueprint variable editor panel

Change 3116367 on 2016/09/07 by [REDACTED_PERSON]

	Fixed Function/Macro inputs/outputs list (had become cramped with my last change) + misc. fixes for new container types, fixes uninitialized members in FTerminalType

Change 3116663 on 2016/09/07 by [REDACTED_PERSON]

	Fix for array functions showing up with TSet and TMap pins

Change 3118259 on 2016/09/08 by [REDACTED_PERSON]

	BP Profiler: Fixed Assert when profiling parent/child Blueprint

Change 3119023 on 2016/09/09 by [REDACTED_PERSON]

	Manually integrated (from Odin branch) recent changes related to BP and nativization:

	3115713 UE-35448
	3117590 UE-35697
	3117742 ODIN-577

Change 3119058 on 2016/09/09 by [REDACTED_PERSON]

	#2574 https://github.com/EpicGames/UnrealEngine/pull/2574
	Renamed function CustomNativeInitilize to InitializeNativeClassData  and made it private.

Change 3119302 on 2016/09/09 by [REDACTED_PERSON]

	Global variable for WITH_PERFCOUNTERS definition in UEBuildConfiguration.
	Previously the same header could be compiled with the WITH_PERFCOUNTERS flag enadles and disabled  (during a single compilation) .

Change 3119502 on 2016/09/09 by [REDACTED_PERSON]

	When building a deterministic UUID for latent nodes, we now use expanded nodes' origin (node) to avoid collisions (latent node in macros, etc.)

Change 3119517 on 2016/09/09 by [REDACTED_PERSON]

	Added blueprint editor settings option to display unique names for blueprint nodes

Change 3119602 on 2016/09/09 by [REDACTED_PERSON]

	Added stats about nativized AnimBP
	Mechanism to exlcude reducible AnimBP
	Editor config option:[BlueprintNativizationSettings] bNativizeAnimBPOnlyWhenNonReducibleFuncitons=false

Change 3119615 on 2016/09/09 by [REDACTED_PERSON]

	Missing change (should be part of cl#3119602)

Change 3119619 on 2016/09/09 by [REDACTED_PERSON]

	Excluding all AnimBP from Orion nativization.

Change 3120752 on 2016/09/12 by [REDACTED_PERSON]

	Removed unnecessary ensure

Change 3121354 on 2016/09/12 by [REDACTED_PERSON]

	Fixed variable type width, required for TMap's extra combobox.

Change 3121626 on 2016/09/12 by [REDACTED_PERSON]

	[UE-35456] Fix crash on right-click in components tree view after copying one or more BSP actors to clipboard.

	Note: This applies to the components tree view in both the Blueprint editor and the Level editor's Actor details panel.

	change summary:
	- modified FComponentObjectTextFactory::CanCreateClass() to exclude Actor/Component subtypes that are not Blueprint-compatible (e.g. ABrush).

Change 3122712 on 2016/09/13 by [REDACTED_PERSON]

	Replaced "check" with "ensure".

Change 3124398 on 2016/09/14 by [REDACTED_PERSON]

	More strict BP validation in UBlueprintThumbnailRenderer::Draw

Change 3124405 on 2016/09/14 by [REDACTED_PERSON]

	Function Libraries are properly added to dependencies list while nativization.

Change 3124667 on 2016/09/14 by [REDACTED_PERSON]

	Fixed incompatible pins validation.

Change 3125245 on 2016/09/14 by [REDACTED_PERSON]

	[UE-33674] Fix missing stats for the ForEachElementInEnum node type in the Blueprint profiler tree view.

	change summary:
	- modified FScriptEventPlayback::Process() to not allow intermediate node exit pins to pollute the current trace path
	- modified FBlueprintFunctionContext::DetermineGraphNodeCharacteristics() to handle the UK2Node_ForEachElementInEnum type as a special case and account for extra loop iterations in the sample frequency computed at mapping time
	- exported UK2Node_ForEachElementInEnum::InsideLoopPinName and EnumOutputPinName string constants

Change 3126211 on 2016/09/15 by [REDACTED_PERSON]

Change 3126393 on 2016/09/15 by [REDACTED_PERSON]

	Replace "check" by "ensure".

Change 3126623 on 2016/09/15 by [REDACTED_PERSON]

	STRUCT_SerializeFromMismatchedTag is not necessary to serialize structure when guids match. Anyway STRUCT_SerializeFromMismatchedTag sholud precede SerializeFromMismatchedTag().

Change 3127288 on 2016/09/15 by [REDACTED_PERSON]

	Making the script VM overhead and native time stats threadsafe (to account for threaded anim Blueprints in Orion).

Change 3127375 on 2016/09/15 by [REDACTED_PERSON]

	Making sure Blueprint classes inherit the super's ClassConfigName properly (inherit the ID instead of the filename).

Change 3127381 on 2016/09/15 by [REDACTED_PERSON]

	Removing an overzealous ensure that certain users were hitting when a loading array property wasn't fully filled out yet (confirmed that it was populated with the proper objects by the end of the load).

Change 3127476 on 2016/09/15 by [REDACTED_PERSON]

	Build fix

Change 3128335 on 2016/09/16 by [REDACTED_PERSON]

	Fixed broken BP assets.

Change 3128589 on 2016/09/16 by [REDACTED_PERSON]

	Fixing a static analysis CIS warning (duplicated condition).

Change 3128630 on 2016/09/16 by [REDACTED_PERSON]

	Re-fix with engine version set

Change 3129338 on 2016/09/16 by [REDACTED_PERSON]

	=FScriptSet/FScriptSetHelper fleshed out (Add, Remove, and Find implemented)
	+SetParam implemented for marking up sets for primitive Set functions (to be checked in once completed as BlueprintSetLibrary)

diff --git a/Engine/Config/BaseEditor.ini b/Engine/Config/BaseEditor.ini
--- a/Engine/Config/BaseEditor.ini
+++ b/Engine/Config/BaseEditor.ini
@@ -475,6 +475,7 @@ bSkipUFUNCTION=false
 bUsePRAGMA_DISABLE_OPTIMIZATION=false
 bDontNativizeDataOnlyBP=false
 bUseLocalEventGraphVariables=true
+bNativizeAnimBPOnlyWhenNonReducibleFuncitons=false
 
 +ClassesWithStaticSearchableValues=/Script/GameplayAbilities.GameplayCueNotify_Actor
 +ClassesWithStaticSearchableValues=/Script/GameplayAbilities.GameplayCueNotify_Static
diff --git a/Engine/Source/Developer/BlueprintCompilerCppBackend/Private/BlueprintCompilerCppBackend.cpp b/Engine/Source/Developer/BlueprintCompilerCppBackend/Private/BlueprintCompilerCppBackend.cpp
--- a/Engine/Source/Developer/BlueprintCompilerCppBackend/Private/BlueprintCompilerCppBackend.cpp
+++ b/Engine/Source/Developer/BlueprintCompilerCppBackend/Private/BlueprintCompilerCppBackend.cpp
@@ -739,6 +739,8 @@ FString FBlueprintCompilerCppBackend::EmitCallStatmentInner(FEmitterLocalContext
 	const bool bIsSpecialCase = HandleSpeciallyNativizedFunction();
 	if (!bIsSpecialCase)
 	{
+		FNativizationSummaryHelper::FunctionUsed(CurrentClass, Statement.FunctionToCall);
+
 		// Emit object to call the method on
 		if (bInterfaceCallExecute)
 		{
@@ -871,6 +873,31 @@ FString FBlueprintCompilerCppBackend::TermToText(FEmitterLocalContext& EmitterCo
 	}
 	else
 	{
+		auto GenerateDefaultLocalVariable = [&](const FBPTerminal* InTerm) -> FString
+		{
+			const FString DefaultValueVariable = EmitterContext.GenerateUniqueLocalName();
+			const uint32 PropertyExportFlags = EPropertyExportCPPFlags::CPPF_CustomTypeName | EPropertyExportCPPFlags::CPPF_BlueprintCppBackend | EPropertyExportCPPFlags::CPPF_NoConst;
+			const FString CppType = Term->AssociatedVarProperty
+				? EmitterContext.ExportCppDeclaration(Term->AssociatedVarProperty, EExportedDeclaration::Local, PropertyExportFlags, true)
+				: FEmitHelper::PinTypeToNativeType(Term->Type);
+
+			const FString DefaultValueConstructor = (!Term->Type.bIsArray)
+				? FEmitHelper::LiteralTerm(EmitterContext, Term->Type, FString(), nullptr, &FText::GetEmpty())
+				: FString::Printf(TEXT("%s{}"), *CppType);
+
+			EmitterContext.AddLine(*FString::Printf(TEXT("%s %s = %s;"), *CppType, *DefaultValueVariable, *DefaultValueConstructor));
+
+			return DefaultValueVariable;
+		};
+
+		if (Term->AssociatedVarProperty && Term->AssociatedVarProperty->HasAnyPropertyFlags(CPF_EditorOnly))
+		{
+			UE_LOG(LogK2Compiler, Warning, TEXT("C++ backend cannot cannot use EditorOnly property: %s"), *GetPathNameSafe(Term->AssociatedVarProperty));
+			EmitterContext.AddLine(*FString::Printf(TEXT("// EDITOR-ONLY Variable: %s"), *FEmitHelper::GetCppName(Term->AssociatedVarProperty)));
+			const FString DefaultValueVariable = GenerateDefaultLocalVariable(Term);
+			return DefaultValueVariable;
+		}
+
 		FString ContextStr;
 		if ((Term->Context != nullptr) && (Term->Context->Name != PSC_Self))
 		{
@@ -915,6 +942,8 @@ FString FBlueprintCompilerCppBackend::TermToText(FEmitterLocalContext& EmitterCo
 		}
 		else if (Term->AssociatedVarProperty)
 		{
+			FNativizationSummaryHelper::PropertyUsed(EmitterContext.GetCurrentlyGeneratedClass(), Term->AssociatedVarProperty);
+
 			const bool bSelfContext = (!Term->Context) || (Term->Context->Name == PSC_Self);
 			const bool bPropertyOfParent = EmitterContext.Dependencies.GetActualStruct()->IsChildOf(Term->AssociatedVarProperty->GetOwnerStruct());
 			bIsAccessible &= !Term->AssociatedVarProperty->HasAnyPropertyFlags(CPF_NativeAccessSpecifierPrivate)
@@ -1001,17 +1030,7 @@ FString FBlueprintCompilerCppBackend::TermToText(FEmitterLocalContext& EmitterCo
 
 		if (!Conditions.IsEmpty())
 		{
-			const FString DefaultValueVariable = EmitterContext.GenerateUniqueLocalName();
-			const uint32 PropertyExportFlags = EPropertyExportCPPFlags::CPPF_CustomTypeName | EPropertyExportCPPFlags::CPPF_BlueprintCppBackend | EPropertyExportCPPFlags::CPPF_NoConst;
-			const FString CppType = Term->AssociatedVarProperty 
-				? EmitterContext.ExportCppDeclaration(Term->AssociatedVarProperty, EExportedDeclaration::Local, PropertyExportFlags, true)
-				: FEmitHelper::PinTypeToNativeType(Term->Type);
-
-			const FString DefaultValueConstructor = (!Term->Type.bIsArray)
-				? FEmitHelper::LiteralTerm(EmitterContext, Term->Type, FString(), nullptr, &FText::GetEmpty())
-				: FString::Printf(TEXT("%s{}"), *CppType);
-
-			EmitterContext.AddLine(*FString::Printf(TEXT("%s %s = %s;"), *CppType, *DefaultValueVariable, *DefaultValueConstructor));
+			const FString DefaultValueVariable = GenerateDefaultLocalVariable(Term);
 			const FString DefaultExpression = bUseWeakPtrGetter ? (DefaultValueVariable + WeakPtrGetter) : DefaultValueVariable;
 			return FString::Printf(TEXT("((%s) ? (%s) : (%s))"), *Conditions, *ResultPath, *DefaultExpression);
 		}
@@ -1056,7 +1075,7 @@ FString FBlueprintCompilerCppBackend::LatentFunctionInfoTermToText(FEmitterLocal
 	return FEmitHelper::LiteralTerm(EmitterContext, Term->Type, StructValues, nullptr);
 }
 
-void FBlueprintCompilerCppBackend::InnerFunctionImplementation(FKismetFunctionContext& FunctionContext, FEmitterLocalContext& EmitterContext, int32 ExecutionGroup)
+bool FBlueprintCompilerCppBackend::InnerFunctionImplementation(FKismetFunctionContext& FunctionContext, FEmitterLocalContext& EmitterContext, int32 ExecutionGroup)
 {
 	EmitterContext.ResetPropertiesForInaccessibleStructs();
 
@@ -1136,7 +1155,7 @@ void FBlueprintCompilerCppBackend::InnerFunctionImplementation(FKismetFunctionCo
 		}
 	}
 
-	EmitAllStatements(FunctionContext, ExecutionGroup, EmitterContext, *ActualLinearExecutionList);
+	const bool bIsNotReducible = EmitAllStatements(FunctionContext, ExecutionGroup, EmitterContext, *ActualLinearExecutionList);
 
 	if (bUseGotoState)
 	{
@@ -1155,6 +1174,8 @@ void FBlueprintCompilerCppBackend::InnerFunctionImplementation(FKismetFunctionCo
 		EmitterContext.DecreaseIndent();
 		EmitterContext.AddLine(TEXT("} while( CurrentState != -1 );"));
 	}
+
+	return bIsNotReducible;
 }
 
 bool FBlueprintCompilerCppBackend::SortNodesInUberGraphExecutionGroup(FKismetFunctionContext &FunctionContext, UEdGraphNode* TheOnlyEntryPoint, int32 ExecutionGroup, TArray<UEdGraphNode*> &LocalLinearExecutionList)
@@ -1325,10 +1346,12 @@ void FBlueprintCompilerCppBackend::EmitStatement(FBlueprintCompiledStatement &St
 	};
 }
 
-void FBlueprintCompilerCppBackend::EmitAllStatements(FKismetFunctionContext &FunctionContext, int32 ExecutionGroup, FEmitterLocalContext &EmitterContext, const TArray<UEdGraphNode*>& LinearExecutionList)
+bool FBlueprintCompilerCppBackend::EmitAllStatements(FKismetFunctionContext &FunctionContext, int32 ExecutionGroup, FEmitterLocalContext &EmitterContext, const TArray<UEdGraphNode*>& LinearExecutionList)
 {
 	ensure(!bUseExecutionGroup || FunctionContext.UnsortedSeparateExecutionGroups.IsValidIndex(ExecutionGroup));
 	bool bFirsCase = true;
+
+	bool bAnyNonReducableStatement = false;
 	// Emit code in the order specified by the linear execution list (the first node is always the entry point for the function)
 	for (int32 NodeIndex = 0; NodeIndex < LinearExecutionList.Num(); ++NodeIndex)
 	{
@@ -1364,9 +1387,11 @@ void FBlueprintCompilerCppBackend::EmitAllStatements(FKismetFunctionContext &Fun
 					EmitterContext.IncreaseIndent();
 				}
 				EmitStatement(Statement, EmitterContext, FunctionContext);
+				bAnyNonReducableStatement |= !FKismetCompilerUtilities::IsStatementReducible(Statement.Type);
 			}
 		}
 	}
+	return bAnyNonReducableStatement;
 }
 
 bool FBlueprintCompilerCppBackend::PrepareToUseExecutionGroupWithoutGoto(FKismetFunctionContext &FunctionContext, int32 ExecutionGroup, UEdGraphNode* &TheOnlyEntryPoint)
diff --git a/Engine/Source/Developer/BlueprintCompilerCppBackend/Private/BlueprintCompilerCppBackend.h b/Engine/Source/Developer/BlueprintCompilerCppBackend/Private/BlueprintCompilerCppBackend.h
--- a/Engine/Source/Developer/BlueprintCompilerCppBackend/Private/BlueprintCompilerCppBackend.h
+++ b/Engine/Source/Developer/BlueprintCompilerCppBackend/Private/BlueprintCompilerCppBackend.h
@@ -28,13 +28,14 @@ public:
 	{}
 
 protected:
-	virtual void InnerFunctionImplementation(FKismetFunctionContext& FunctionContext, FEmitterLocalContext& EmitterContext, int32 ExecutionGroup) override;
+	virtual bool InnerFunctionImplementation(FKismetFunctionContext& FunctionContext, FEmitterLocalContext& EmitterContext, int32 ExecutionGroup) override;
 
 	// creates local linear execution list ,returns if the execution group can be handled without switch
 	bool SortNodesInUberGraphExecutionGroup(FKismetFunctionContext &FunctionContext, UEdGraphNode* TheOnlyEntryPoint, int32 ExecutionGroup, TArray<UEdGraphNode*> &LocalLinearExecutionList);
 	// return if the execution group can be handled without switch
 	bool PrepareToUseExecutionGroupWithoutGoto(FKismetFunctionContext &FunctionContext, int32 ExecutionGroup, UEdGraphNode* &TheOnlyEntryPoint);
-	void EmitAllStatements(FKismetFunctionContext &FunctionContext, int32 ExecutionGroup, FEmitterLocalContext &EmitterContext, const TArray<UEdGraphNode*>& LinearExecutionList);
+	// returns if the function performs any significant action (it is not reducible)
+	bool EmitAllStatements(FKismetFunctionContext &FunctionContext, int32 ExecutionGroup, FEmitterLocalContext &EmitterContext, const TArray<UEdGraphNode*>& LinearExecutionList);
 	void EmitStatement(FBlueprintCompiledStatement &Statement, FEmitterLocalContext &EmitterContext, FKismetFunctionContext& FunctionContext);
 
 protected:
diff --git a/Engine/Source/Developer/BlueprintCompilerCppBackend/Private/BlueprintCompilerCppBackendBase.cpp b/Engine/Source/Developer/BlueprintCompilerCppBackend/Private/BlueprintCompilerCppBackendBase.cpp
--- a/Engine/Source/Developer/BlueprintCompilerCppBackend/Private/BlueprintCompilerCppBackendBase.cpp
+++ b/Engine/Source/Developer/BlueprintCompilerCppBackend/Private/BlueprintCompilerCppBackendBase.cpp
@@ -185,6 +185,11 @@ FString FBlueprintCompilerCppBackendBase::GenerateCodeFromClass(UClass* SourceCl
 	FGatherConvertedClassDependencies Dependencies(SourceClass);
 	FEmitterLocalContext EmitterContext(Dependencies);
 
+	UClass* OriginalSourceClass = Dependencies.FindOriginalClass(SourceClass);
+	ensure(OriginalSourceClass != SourceClass);
+
+	FNativizationSummaryHelper::RegisterClass(OriginalSourceClass);
+
 	EmitFileBeginning(CleanCppClassName, EmitterContext);
 
 	// C4883 is a strange error (for big functions), introduced in VS2015 update 2
@@ -219,9 +224,6 @@ FString FBlueprintCompilerCppBackendBase::GenerateCodeFromClass(UClass* SourceCl
 		}
 	}
 
-	UClass* OriginalSourceClass = Dependencies.FindOriginalClass(SourceClass);
-	ensure(OriginalSourceClass != SourceClass);
-
 	const bool bHasStaticSearchableValues = FBackendHelperStaticSearchableValues::HasSearchableValues(SourceClass);
 
 	// Class declaration
@@ -725,7 +727,7 @@ void FBlueprintCompilerCppBackendBase::ConstructFunctionBody(FEmitterLocalContex
 {
 	if (FunctionContext.UnsortedSeparateExecutionGroups.Num() && (ExecutionGroup < 0))
 	{
-		// uso only for latent actions..
+		// use only for latent actions..
 		return;
 	}
 
@@ -751,7 +753,11 @@ void FBlueprintCompilerCppBackendBase::ConstructFunctionBody(FEmitterLocalContex
 		}
 	}
 
-	InnerFunctionImplementation(FunctionContext, EmitterContext, ExecutionGroup);
+	bool bIsFunctionNotReducible = InnerFunctionImplementation(FunctionContext, EmitterContext, ExecutionGroup);
+	if (!bIsFunctionNotReducible)
+	{
+		FNativizationSummaryHelper::ReducibleFunciton(EmitterContext.Dependencies.FindOriginalClass(EmitterContext.GetCurrentlyGeneratedClass()));
+	}
 }
 
 FString FBlueprintCompilerCppBackendBase::GenerateCodeFromEnum(UUserDefinedEnum* SourceEnum)
diff --git a/Engine/Source/Developer/BlueprintCompilerCppBackend/Private/BlueprintCompilerCppBackendBase.h b/Engine/Source/Developer/BlueprintCompilerCppBackend/Private/BlueprintCompilerCppBackendBase.h
--- a/Engine/Source/Developer/BlueprintCompilerCppBackend/Private/BlueprintCompilerCppBackendBase.h
+++ b/Engine/Source/Developer/BlueprintCompilerCppBackend/Private/BlueprintCompilerCppBackendBase.h
@@ -61,8 +61,10 @@ protected:
 				return;
 			}
 
+		Returns true is the generated function is not reducible.
+
 	*/
-	virtual void InnerFunctionImplementation(FKismetFunctionContext& FunctionContext, FEmitterLocalContext& EmitterContext, int32 ExecutionGroup) = 0;
+	virtual bool InnerFunctionImplementation(FKismetFunctionContext& FunctionContext, FEmitterLocalContext& EmitterContext, int32 ExecutionGroup) = 0;
 	
 	int32 StatementToStateIndex(FKismetFunctionContext& FunctionContext, FBlueprintCompiledStatement* Statement)
 	{
diff --git a/Engine/Source/Developer/BlueprintCompilerCppBackend/Private/BlueprintCompilerCppBackendGatherDependencies.cpp b/Engine/Source/Developer/BlueprintCompilerCppBackend/Private/BlueprintCompilerCppBackendGatherDependencies.cpp
--- a/Engine/Source/Developer/BlueprintCompilerCppBackend/Private/BlueprintCompilerCppBackendGatherDependencies.cpp
+++ b/Engine/Source/Developer/BlueprintCompilerCppBackend/Private/BlueprintCompilerCppBackendGatherDependencies.cpp
@@ -81,8 +81,21 @@ struct FFindAssetsToInclude : public FGatherConvertedClassDependenciesHelperBase
 			return;
 		}
 
-		const bool bUseZConstructorInGeneratedCode = false;
+		if (Object->HasAnyFlags(RF_ClassDefaultObject))
+		{
+			// Static functions from libraries are called on CDO. (The functions is stored as a name not an object).
+			UClass* OwnerClass = Object->GetClass();
+			if (OwnerClass && (OwnerClass != CurrentlyConvertedStruct))
+			{
+				UBlueprintGeneratedClass* OwnerAsBPGC = Cast<UBlueprintGeneratedClass>(OwnerClass);
+				if (OwnerAsBPGC && !Dependencies.ConvertedClasses.Contains(OwnerAsBPGC) && Dependencies.WillClassBeConverted(OwnerAsBPGC))
+				{
+					Dependencies.ConvertedClasses.Add(OwnerAsBPGC);
+				}
+			}
+		}
 
+		const bool bUseZConstructorInGeneratedCode = false;
 		//TODO: What About Delegates?
 		auto ObjAsBPGC = Cast<UBlueprintGeneratedClass>(Object);
 		const bool bWillBeConvetedAsBPGC = ObjAsBPGC && Dependencies.WillClassBeConverted(ObjAsBPGC);
diff --git a/Engine/Source/Developer/BlueprintCompilerCppBackend/Private/BlueprintCompilerCppBackendUMG.cpp b/Engine/Source/Developer/BlueprintCompilerCppBackend/Private/BlueprintCompilerCppBackendUMG.cpp
--- a/Engine/Source/Developer/BlueprintCompilerCppBackend/Private/BlueprintCompilerCppBackendUMG.cpp
+++ b/Engine/Source/Developer/BlueprintCompilerCppBackend/Private/BlueprintCompilerCppBackendUMG.cpp
@@ -12,7 +12,7 @@ void FBackendHelperUMG::WidgetFunctionsInHeader(FEmitterLocalContext& Context)
 	{
 		Context.Header.AddLine(FString::Printf(TEXT("virtual void %s(TArray<FName>& SlotNames) const override;"), GET_FUNCTION_NAME_STRING_CHECKED(UUserWidget, GetSlotNames)));
 		Context.Header.AddLine(FString::Printf(TEXT("virtual void %s(const class ITargetPlatform* TargetPlatform) override;"), GET_FUNCTION_NAME_STRING_CHECKED(UUserWidget, PreSave)));
-		Context.Header.AddLine(FString::Printf(TEXT("virtual void %s() override;"), GET_FUNCTION_NAME_STRING_CHECKED(UUserWidget, CustomNativeInitilize)));
+		Context.Header.AddLine(TEXT("virtual void InitializeNativeClassData() override;"));
 	}
 }
 
@@ -74,8 +74,8 @@ void FBackendHelperUMG::EmitWidgetInitializationFunctions(FEmitterLocalContext&
 			Context.AddLine(TEXT("}"));
 		}
 
-		{	// CustomNativeInitilize
-			Context.AddLine(FString::Printf(TEXT("void %s::%s()"), *CppClassName, GET_FUNCTION_NAME_STRING_CHECKED(UUserWidget, CustomNativeInitilize)));
+		{	// InitializeNativeClassData
+			Context.AddLine(FString::Printf(TEXT("void %s::InitializeNativeClassData()"), *CppClassName));
 			Context.AddLine(TEXT("{"));
 			Context.IncreaseIndent();
 
diff --git a/Engine/Source/Developer/BlueprintCompilerCppBackend/Private/BlueprintCompilerCppBackendUtils.cpp b/Engine/Source/Developer/BlueprintCompilerCppBackend/Private/BlueprintCompilerCppBackendUtils.cpp
--- a/Engine/Source/Developer/BlueprintCompilerCppBackend/Private/BlueprintCompilerCppBackendUtils.cpp
+++ b/Engine/Source/Developer/BlueprintCompilerCppBackend/Private/BlueprintCompilerCppBackendUtils.cpp
@@ -1027,9 +1027,10 @@ FString FEmitHelper::LiteralTerm(FEmitterLocalContext& EmitterContext, const FEd
 				}
 
 				{
-					const FString StructOnScopeName = EmitterContext.GenerateUniqueLocalName();
-					EmitterContext.AddLine(FString::Printf(TEXT("FStructOnScope %s(%s::StaticStruct());"), *StructOnScopeName, *StructName));
-					EmitterContext.AddLine(FString::Printf(TEXT("%s& %s = *((%s*)%s.GetStructMemory());"), *StructName, *LocalStructNativeName, *StructName, *StructOnScopeName));  // TODO: ?? should "::GetDefaultValue()" be called here?
+					const FString StructMemoryVar = EmitterContext.GenerateUniqueLocalName();
+					EmitterContext.AddLine(FString::Printf(TEXT("uint8* %s = (uint8*)FMemory_Alloca(%s::StaticStruct()->GetStructureSize());"), *StructMemoryVar, *StructName));
+					EmitterContext.AddLine(FString::Printf(TEXT("%s::StaticStruct()->InitializeStruct(%s);"), *StructName, *StructMemoryVar));
+					EmitterContext.AddLine(FString::Printf(TEXT("%s& %s = *reinterpret_cast<%s*>(%s);"), *StructName, *LocalStructNativeName, *StructName, *StructMemoryVar));  // TODO: ?? should "::GetDefaultValue()" be called here?
 				}
 
 				{
@@ -1395,7 +1396,8 @@ FString FEmitHelper::GenerateGetPropertyByName(FEmitterLocalContext& EmitterCont
 	return PropertyPtrName;
 }
 
-static void UpdateNativizationSummary_InaccessibleProperty(const UProperty* Property)
+
+void FNativizationSummaryHelper::InaccessibleProperty(const UProperty* Property)
 {
 	IBlueprintCompilerCppBackendModule& BackEndModule = (IBlueprintCompilerCppBackendModule&)IBlueprintCompilerCppBackendModule::Get();
 	TSharedPtr<FNativizationSummary> NativizationSummary = BackEndModule.NativizationSummary();
@@ -1414,6 +1416,85 @@ static void UpdateNativizationSummary_InaccessibleProperty(const UProperty* Prop
 	}
 }
 
+static void FNativizationSummaryHelper__MemberUsed(const UClass* Class, const UField* Field, int32 FNativizationSummary::FAnimBlueprintDetails::*CouterPtr)
+{
+	if (Field && Class)
+	{
+		IBlueprintCompilerCppBackendModule& BackEndModule = (IBlueprintCompilerCppBackendModule&)IBlueprintCompilerCppBackendModule::Get();
+		TSharedPtr<FNativizationSummary> NativizationSummary = BackEndModule.NativizationSummary();
+		if (NativizationSummary.IsValid())
+		{
+			const UClass* Owner = Field->GetOwnerClass();
+			UAnimBlueprint* AnimBP = Cast<UAnimBlueprint>(UBlueprint::GetBlueprintFromClass(Owner));
+			const bool bUnrelatedClass = !Class->IsChildOf(Owner);
+			if (AnimBP && bUnrelatedClass)
+			{
+				FNativizationSummary::FAnimBlueprintDetails& AnimBlueprintDetails = NativizationSummary->AnimBlueprintStat.FindOrAdd(FStringAssetReference(AnimBP));
+				(AnimBlueprintDetails.*CouterPtr)++;
+			}
+		}
+	}
+}
+
+void FNativizationSummaryHelper::PropertyUsed(const UClass* Class, const UProperty* Property)
+{
+	FNativizationSummaryHelper__MemberUsed(Class, Property, &FNativizationSummary::FAnimBlueprintDetails::VariableUsage);
+}
+
+void FNativizationSummaryHelper::FunctionUsed(const UClass* Class, const UFunction* Function)
+{
+	FNativizationSummaryHelper__MemberUsed(Class, Function, &FNativizationSummary::FAnimBlueprintDetails::FunctionUsage);
+}
+
+void FNativizationSummaryHelper::ReducibleFunciton(const UClass* OriginalClass)
+{
+	if (OriginalClass)
+	{
+		IBlueprintCompilerCppBackendModule& BackEndModule = (IBlueprintCompilerCppBackendModule&)IBlueprintCompilerCppBackendModule::Get();
+		TSharedPtr<FNativizationSummary> NativizationSummary = BackEndModule.NativizationSummary();
+		UAnimBlueprint* AnimBP = Cast<UAnimBlueprint>(UBlueprint::GetBlueprintFromClass(OriginalClass));
+		if (NativizationSummary.IsValid() && OriginalClass && AnimBP)
+		{
+			FNativizationSummary::FAnimBlueprintDetails& AnimBlueprintDetails = NativizationSummary->AnimBlueprintStat.FindOrAdd(FStringAssetReference(AnimBP));
+			AnimBlueprintDetails.ReducibleFunctions++;
+		}
+	}
+}
+
+void FNativizationSummaryHelper::RegisterClass(const UClass* OriginalClass)
+{
+	IBlueprintCompilerCppBackendModule& BackEndModule = (IBlueprintCompilerCppBackendModule&)IBlueprintCompilerCppBackendModule::Get();
+	TSharedPtr<FNativizationSummary> NativizationSummary = BackEndModule.NativizationSummary();
+	UAnimBlueprint* AnimBP = Cast<UAnimBlueprint>(UBlueprint::GetBlueprintFromClass(OriginalClass));
+	if (NativizationSummary.IsValid() && OriginalClass && AnimBP)
+	{
+		{
+			FNativizationSummary::FAnimBlueprintDetails& AnimBlueprintDetails = NativizationSummary->AnimBlueprintStat.FindOrAdd(FStringAssetReference(AnimBP));
+
+			AnimBlueprintDetails.Variables = AnimBP->NewVariables.Num();
+
+			UFunction* UberGraphFunction = CastChecked<UBlueprintGeneratedClass>(OriginalClass)->UberGraphFunction;
+			for (auto FunctIter : TFieldRange<UFunction>(OriginalClass, EFieldIteratorFlags::ExcludeSuper))
+			{
+				if (UberGraphFunction != FunctIter)
+				{
+					AnimBlueprintDetails.Functions++;
+				}
+			}
+		}
+
+		for (UClass* SuperClass = OriginalClass->GetSuperClass(); SuperClass; SuperClass = SuperClass->GetSuperClass())
+		{
+			UAnimBlueprint* ParentAnimBP = Cast<UAnimBlueprint>(UBlueprint::GetBlueprintFromClass(SuperClass));
+			if (ParentAnimBP)
+			{
+				FNativizationSummary::FAnimBlueprintDetails& AnimBlueprintDetails = NativizationSummary->AnimBlueprintStat.FindOrAdd(FStringAssetReference(ParentAnimBP));
+				AnimBlueprintDetails.Children++;
+			}
+		}
+	}
+}
+
 FString FEmitHelper::AccessInaccessibleProperty(FEmitterLocalContext& EmitterContext, const UProperty* Property
 	, const FString& ContextStr, const FString& ContextAdressOp, int32 StaticArrayIdx, ENativizedTermUsage TermUsage, FString* CustomSetExpressionEnding)
 {
@@ -1430,7 +1511,7 @@ FString FEmitHelper::AccessInaccessibleProperty(FEmitterLocalContext& EmitterCon
 	{
 		if (TermUsage == ENativizedTermUsage::Getter)
 		{
-			UpdateNativizationSummary_InaccessibleProperty(Property);
+			FNativizationSummaryHelper::InaccessibleProperty(Property);
 			const FString PropertyLocalName = GenerateGetPropertyByName(EmitterContext, Property);
 			return  FString::Printf(TEXT("(((UBoolProperty*)%s)->%s(%s(%s), %d))")
 				, *PropertyLocalName
@@ -1442,7 +1523,7 @@ FString FEmitHelper::AccessInaccessibleProperty(FEmitterLocalContext& EmitterCon
 
 		if (TermUsage == ENativizedTermUsage::Setter)
 		{
-			UpdateNativizationSummary_InaccessibleProperty(Property);
+			FNativizationSummaryHelper::InaccessibleProperty(Property);
 			const FString PropertyLocalName = GenerateGetPropertyByName(EmitterContext, Property);
 			if (ensure(CustomSetExpressionEnding))
 			{
@@ -1468,7 +1549,7 @@ FString FEmitHelper::AccessInaccessibleProperty(FEmitterLocalContext& EmitterCon
 	{
 		//TODO: if property is inaccessible due to const specifier, use const_cast
 
-		UpdateNativizationSummary_InaccessibleProperty(Property);
+		FNativizationSummaryHelper::InaccessibleProperty(Property);
 		const FString PropertyLocalName = GenerateGetPropertyByName(EmitterContext, Property);
 		return FString::Printf(TEXT("(*(%s->ContainerPtrToValuePtr<%s>(%s(%s), %d)))")
 			, *PropertyLocalName
@@ -1583,7 +1664,7 @@ FString FEmitHelper::AccessInaccessiblePropertyUsingOffset(FEmitterLocalContext&
 		UE_LOG(LogK2Compiler, Warning, TEXT("AccessInaccessiblePropertyUsingOffset - NOEXPORT structure should be handled in a custom way: %s"), *GetPathNameSafe(Property->GetOwnerStruct()));
 	}
 
-	UpdateNativizationSummary_InaccessibleProperty(Property);
+	FNativizationSummaryHelper::InaccessibleProperty(Property);
 
 	const int32 PropertyOffsetWithoutEditorOnlyMembers = HelperWithoutEditorOnlyMembers::OffsetWithoutEditorOnlyMembers(Property);
 	/*
diff --git a/Engine/Source/Developer/BlueprintCompilerCppBackend/Private/BlueprintCompilerCppBackendUtils.h b/Engine/Source/Developer/BlueprintCompilerCppBackend/Private/BlueprintCompilerCppBackendUtils.h
--- a/Engine/Source/Developer/BlueprintCompilerCppBackend/Private/BlueprintCompilerCppBackendUtils.h
+++ b/Engine/Source/Developer/BlueprintCompilerCppBackend/Private/BlueprintCompilerCppBackendUtils.h
@@ -292,7 +292,7 @@ struct FEmitDefaultValueHelper
 
 	// Creates the subobject (of class) returns it's native local name, 
 	// returns empty string if cannot handle
-	static FString HandleClassSubobject(FEmitterLocalContext& Context, UObject* Object, FEmitterLocalContext::EClassSubobjectList ListOfSubobjectsTyp, bool bCreate, bool bInitilize);
+	static FString HandleClassSubobject(FEmitterLocalContext& Context, UObject* Object, FEmitterLocalContext::EClassSubobjectList ListOfSubobjectsTyp, bool bCreate, bool bInitialize);
 
 	// returns true, and fill OutResult, when the structure is handled in a custom way.
 	static bool SpecialStructureConstructor(const UStruct* Struct, const uint8* ValuePtr, FString* OutResult);
@@ -339,4 +339,16 @@ struct FBackendHelperStaticSearchableValues
 	static FString GenerateClassMetaData(UClass* Class);
 	static void EmitFunctionDeclaration(FEmitterLocalContext& Context);
 	static void EmitFunctionDefinition(FEmitterLocalContext& Context);
+};
+
+struct FNativizationSummaryHelper
+{
+	static void InaccessibleProperty(const UProperty* Property);
+	// Notify, that the class used a (unrelated) property
+	static void PropertyUsed(const UClass* Class, const UProperty* Property);
+	// Notify, that the class used a (unrelated) function
+	static void FunctionUsed(const UClass* Class, const UFunction* Function);
+	static void RegisterClass(const UClass* OriginalClass);
+
+	static void ReducibleFunciton(const UClass* OriginalClass);
 };
\ No newline at end of file
diff --git a/Engine/Source/Developer/BlueprintCompilerCppBackend/Private/BlueprintCompilerCppBackendValueHelper.cpp b/Engine/Source/Developer/BlueprintCompilerCppBackend/Private/BlueprintCompilerCppBackendValueHelper.cpp
--- a/Engine/Source/Developer/BlueprintCompilerCppBackend/Private/BlueprintCompilerCppBackendValueHelper.cpp
+++ b/Engine/Source/Developer/BlueprintCompilerCppBackend/Private/BlueprintCompilerCppBackendValueHelper.cpp
@@ -52,6 +52,8 @@ void FEmitDefaultValueHelper::OuterGenerate(FEmitterLocalContext& Context
 			|| ((Property->PropertyFlags & CPF_Config) != 0) 
 			|| (!Property->Identical_InContainer(DataContainer, OptionalDefaultDataContainer, ArrayIndex) && !IsInstancedSubobjectLambda(ArrayIndex)))
 		{
+			FNativizationSummaryHelper::PropertyUsed(Context.GetCurrentlyGeneratedClass(), Property);
+
 			FString PathToMember;
 			UBlueprintGeneratedClass* PropertyOwnerAsBPGC = Cast<UBlueprintGeneratedClass>(Property->GetOwnerClass());
 			UScriptStruct* PropertyOwnerAsScriptStruct = Cast<UScriptStruct>(Property->GetOwnerStruct());
@@ -159,7 +161,7 @@ void FEmitDefaultValueHelper::InnerGenerate(FEmitterLocalContext& Context, const
 			}
 			else if (ValueStr.IsEmpty())
 			{
-				UE_LOG(LogK2Compiler, Error, TEXT("FEmitDefaultValueHelper Cannot generate initilization: %s"), *LocalProperty->GetPathName());
+				UE_LOG(LogK2Compiler, Error, TEXT("FEmitDefaultValueHelper Cannot generate initialization: %s"), *LocalProperty->GetPathName());
 			}
 		}
 		OutSingleLine += ValueStr;
@@ -1057,13 +1059,13 @@ void FEmitDefaultValueHelper::GenerateCustomDynamicClassInitialization(FEmitterL
 
 	Context.AddLine(TEXT("FConvertedBlueprintsDependencies::FillUsedAssetsInDynamicClass(InDynamicClass, &__StaticDependencies_DirectlyUsedAssets);"));
 
-	auto CreateAndInitializeClassSubobjects = [&](bool bCreate, bool bInitilize)
+	auto CreateAndInitializeClassSubobjects = [&](bool bCreate, bool bInitialize)
 	{
 		for (auto ComponentTemplate : ActorComponentTempatesOwnedByClass)
 		{
 			if (ComponentTemplate)
 			{
-				HandleClassSubobject(Context, ComponentTemplate, FEmitterLocalContext::EClassSubobjectList::ComponentTemplates, bCreate, bInitilize);
+				HandleClassSubobject(Context, ComponentTemplate, FEmitterLocalContext::EClassSubobjectList::ComponentTemplates, bCreate, bInitialize);
 			}
 		}
 
@@ -1071,7 +1073,7 @@ void FEmitDefaultValueHelper::GenerateCustomDynamicClassInitialization(FEmitterL
 		{
 			if (TimelineTemplate)
 			{
-				HandleClassSubobject(Context, TimelineTemplate, FEmitterLocalContext::EClassSubobjectList::Timelines, bCreate, bInitilize);
+				HandleClassSubobject(Context, TimelineTemplate, FEmitterLocalContext::EClassSubobjectList::Timelines, bCreate, bInitialize);
 			}
 		}
 
@@ -1079,10 +1081,10 @@ void FEmitDefaultValueHelper::GenerateCustomDynamicClassInitialization(FEmitterL
 		{
 			if (DynamicBindingObject)
 			{
-				HandleClassSubobject(Context, DynamicBindingObject, FEmitterLocalContext::EClassSubobjectList::DynamicBindingObjects, bCreate, bInitilize);
+				HandleClassSubobject(Context, DynamicBindingObject, FEmitterLocalContext::EClassSubobjectList::DynamicBindingObjects, bCreate, bInitialize);
 			}
 		}
-		FBackendHelperUMG::CreateClassSubobjects(Context, bCreate, bInitilize);
+		FBackendHelperUMG::CreateClassSubobjects(Context, bCreate, bInitialize);
 	};
 	CreateAndInitializeClassSubobjects(true, false);
 	CreateAndInitializeClassSubobjects(false, true);
@@ -1278,7 +1280,7 @@ void FEmitDefaultValueHelper::GenerateConstructor(FEmitterLocalContext& Context)
 	Context.AddLine(TEXT("}"));
 }
 
-FString FEmitDefaultValueHelper::HandleClassSubobject(FEmitterLocalContext& Context, UObject* Object, FEmitterLocalContext::EClassSubobjectList ListOfSubobjectsType, bool bCreate, bool bInitilize)
+FString FEmitDefaultValueHelper::HandleClassSubobject(FEmitterLocalContext& Context, UObject* Object, FEmitterLocalContext::EClassSubobjectList ListOfSubobjectsType, bool bCreate, bool bInitialize)
 {
 	ensure(Context.CurrentCodeType == FEmitterLocalContext::EGeneratedCodeType::SubobjectsOfClass);
 
@@ -1288,7 +1290,7 @@ FString FEmitDefaultValueHelper::HandleClassSubobject(FEmitterLocalContext& Cont
 		FString OuterStr = Context.FindGloballyMappedObject(Object->GetOuter());
 		if (OuterStr.IsEmpty())
 		{
-			OuterStr = HandleClassSubobject(Context, Object->GetOuter(), ListOfSubobjectsType, bCreate, bInitilize);
+			OuterStr = HandleClassSubobject(Context, Object->GetOuter(), ListOfSubobjectsType, bCreate, bInitialize);
 			if (OuterStr.IsEmpty())
 			{
 				return FString();
@@ -1322,7 +1324,7 @@ FString FEmitDefaultValueHelper::HandleClassSubobject(FEmitterLocalContext& Cont
 		}
 	}
 
-	if (bInitilize)
+	if (bInitialize)
 	{
 		if (LocalNativeName.IsEmpty())
 		{
diff --git a/Engine/Source/Developer/BlueprintCompilerCppBackend/Public/IBlueprintCompilerCppBackendModule.h b/Engine/Source/Developer/BlueprintCompilerCppBackend/Public/IBlueprintCompilerCppBackendModule.h
--- a/Engine/Source/Developer/BlueprintCompilerCppBackend/Public/IBlueprintCompilerCppBackendModule.h
+++ b/Engine/Source/Developer/BlueprintCompilerCppBackend/Public/IBlueprintCompilerCppBackendModule.h
@@ -7,6 +7,22 @@
 struct FNativizationSummary
 {
 	TMap<FStringAssetReference, int32> InaccessiblePropertyStat;
+
+	struct FAnimBlueprintDetails
+	{
+		int32 Functions; //Number of functions or events
+		int32 ReducibleFunctions; //Number of empty functions
+		int32 Variables; // Number of new BlueprintReadWrite properties
+		int32 Children; // Number of nativized children
+		int32 FunctionUsage; //How many times a function (introduced in the BP) was used by unrelated nativized BP
+		int32 VariableUsage; //How many times a variable (introduced in the BP) was used by unrelated nativized BP
+
+		FAnimBlueprintDetails()
+			: Functions(0), ReducibleFunctions(0), Variables(0), Children(0), FunctionUsage(0), VariableUsage(0)
+		{}
+	};
+
+	TMap<FStringAssetReference, FAnimBlueprintDetails> AnimBlueprintStat;
 };
 
 /**
diff --git a/Engine/Source/Developer/BlueprintNativeCodeGen/Private/BlueprintNativeCodeGenModule.cpp b/Engine/Source/Developer/BlueprintNativeCodeGen/Private/BlueprintNativeCodeGenModule.cpp
--- a/Engine/Source/Developer/BlueprintNativeCodeGen/Private/BlueprintNativeCodeGenModule.cpp
+++ b/Engine/Source/Developer/BlueprintNativeCodeGen/Private/BlueprintNativeCodeGenModule.cpp
@@ -310,11 +310,27 @@ void FBlueprintNativeCodeGenModule::GenerateFullyConvertedClasses()
 	
 	if (NativizationSummary->InaccessiblePropertyStat.Num())
 	{
-		UE_LOG(LogBlueprintCodeGen, Warning, TEXT("Nativization Summary - Inaccessible Properties:"));
+		UE_LOG(LogBlueprintCodeGen, Display, TEXT("Nativization Summary - Inaccessible Properties:"));
 		NativizationSummary->InaccessiblePropertyStat.ValueSort(TGreater<int32>());
 		for (auto& Iter : NativizationSummary->InaccessiblePropertyStat)
 		{
-			UE_LOG(LogBlueprintCodeGen, Warning, TEXT("\t %s \t - %d"), *Iter.Key.ToString(), Iter.Value);
+			UE_LOG(LogBlueprintCodeGen, Display, TEXT("\t %s \t - %d"), *Iter.Key.ToString(), Iter.Value);
+		}
+	}
+	{
+		UE_LOG(LogBlueprintCodeGen, Display, TEXT("Nativization Summary - AnimBP:"));
+		UE_LOG(LogBlueprintCodeGen, Display, TEXT("Name, Children, Non-empty Functions (Empty Functions), Variables, FunctionUsage, VariableUsage"));
+		for (auto& Iter : NativizationSummary->AnimBlueprintStat)
+		{
+			UE_LOG(LogBlueprintCodeGen, Display
+				, TEXT("%s, %d, %d (%d), %d, %d, %d")
+				, *Iter.Key.ToString()
+				, Iter.Value.Children
+				, Iter.Value.Functions - Iter.Value.ReducibleFunctions
+				, Iter.Value.ReducibleFunctions
+				, Iter.Value.Variables
+				, Iter.Value.FunctionUsage
+				, Iter.Value.VariableUsage);
 		}
 	}
 }
@@ -411,7 +427,7 @@ void FBlueprintNativeCodeGenModule::GenerateSingleAsset(UField* ForConversion, c
 		ConversionRecord.GeneratedHeaderPath.Empty();
 	}
 
-	if (ensure(bSuccess))
+	if (bSuccess)
 	{
 		GetManifest(PlatformName).GatherModuleDependencies(ForConversion->GetOutermost());
 	}
@@ -654,6 +670,20 @@ EReplacementResult FBlueprintNativeCodeGenModule::IsTargetedForReplacement(const
 	{
 		if (UBlueprint* Blueprint = Cast<UBlueprint>(BlueprintClass->ClassGeneratedBy))
 		{
+			static const FBoolConfigValueHelper NativizeAnimBPOnlyWhenNonReducibleFuncitons(TEXT("BlueprintNativizationSettings"), TEXT("bNativizeAnimBPOnlyWhenNonReducibleFuncitons"));
+			if (NativizeAnimBPOnlyWhenNonReducibleFuncitons)
+			{
+				if (UAnimBlueprint* AnimBlueprint = Cast<UAnimBlueprint>(Blueprint))
+				{
+					ensure(AnimBlueprint->bHasBeenRegenerated);
+					if (AnimBlueprint->bHasAnyNonReducibleFunction == UBlueprint::EIsBPNonReducible::No)
+					{
+						UE_LOG(LogBlueprintCodeGen, Log, TEXT("AnimBP %s without non-reducible functions is excluded from nativization"), *GetPathNameSafe(Blueprint));
+						Result = EReplacementResult::GenerateStub;
+					}
+				}
+			}
+
 			const EBlueprintType UnconvertableBlueprintTypes[] = {
 				//BPTYPE_Const,		// WTF is a "const" Blueprint?
 				BPTYPE_MacroLibrary,
diff --git a/Engine/Source/Developer/BlueprintProfiler/Private/BlueprintProfiler.cpp b/Engine/Source/Developer/BlueprintProfiler/Private/BlueprintProfiler.cpp
--- a/Engine/Source/Developer/BlueprintProfiler/Private/BlueprintProfiler.cpp
+++ b/Engine/Source/Developer/BlueprintProfiler/Private/BlueprintProfiler.cpp
@@ -64,9 +64,12 @@ void FBlueprintProfiler::InstrumentEvent(const FScriptInstrumentationSignal& Eve
 	#endif
 	{
 		// Handle context switching events
-		CaptureContext.UpdateContext(Event, InstrumentationEventQueue);
-		// Add instrumented event
-		InstrumentationEventQueue.Add(FScriptInstrumentedEvent(Event.GetType(), Event.GetFunctionClassScope()->GetFName(), Event.GetFunctionName(), Event.GetScriptCodeOffset()));
+		if (CaptureContext.UpdateContext(Event, InstrumentationEventQueue))
+		{
+			// Add instrumented event
+			InstrumentationEventQueue.Add(FScriptInstrumentedEvent(Event.GetType(), Event.GetFunctionClassScope()->GetFName(), Event.GetFunctionName(), Event.GetScriptCodeOffset()));
+		}
+
 		// Reset context on event end
 		if (Event.GetType() == EScriptInstrumentation::Stop)
 		{
@@ -81,11 +84,11 @@ void FBlueprintProfiler::AddInstrumentedBlueprint(UBlueprint* InstrumentedBluepr
 {
 	if (InstrumentedBlueprint && InstrumentedBlueprint->GeneratedClass && InstrumentedBlueprint->GeneratedClass->HasInstrumentation())
 	{
-		GetBlueprintContext(InstrumentedBlueprint->GeneratedClass->GetPathName());
+		GetOrCreateBlueprintContext(InstrumentedBlueprint->GeneratedClass->GetPathName());
 	}
 }
 
-TSharedPtr<FBlueprintExecutionContext> FBlueprintProfiler::GetBlueprintContext(const FString& BlueprintClassPath)
+TSharedPtr<FBlueprintExecutionContext> FBlueprintProfiler::GetOrCreateBlueprintContext(const FString& BlueprintClassPath)
 {
 	SCOPE_CYCLE_COUNTER(STAT_BlueprintLookupCost);
 	TSharedPtr<FBlueprintExecutionContext>& Result = PathToBlueprintContext.FindOrAdd(BlueprintClassPath);
@@ -111,6 +114,21 @@ TSharedPtr<FBlueprintExecutionContext> FBlueprintProfiler::GetBlueprintContext(c
 	return Result;
 }
 
+TSharedPtr<FBlueprintExecutionContext> FBlueprintProfiler::FindBlueprintContext(const FString& BlueprintClassPath)
+{
+	SCOPE_CYCLE_COUNTER(STAT_BlueprintLookupCost);
+	TSharedPtr<FBlueprintExecutionContext> Result;
+	if (TSharedPtr<FBlueprintExecutionContext>* SearchResult = PathToBlueprintContext.Find(BlueprintClassPath))
+	{
+		Result = *SearchResult;
+	}
+	else if (TSharedPtr<FBlueprintExecutionContext>* UtilitySearchResult = PathToBlueprintUtilityContext.Find(BlueprintClassPath))
+	{
+		Result = *UtilitySearchResult;
+	}
+	return Result;
+}
+
 TSharedPtr<FScriptExecutionNode> FBlueprintProfiler::GetProfilerDataForNode(const UEdGraphNode* GraphNode)
 {
 	TSharedPtr<FScriptExecutionNode> Result;
@@ -165,7 +183,8 @@ void FBlueprintProfiler::Tick(float DeltaSeconds)
 
 bool FBlueprintProfiler::IsTickable() const
 {
-	return GetDefault<UBlueprintProfilerSettings>()->GetPerformanceThresholdsModified() || InstrumentationEventQueue.Num() > 0;
+	const UBlueprintProfilerSettings* Settings = GetDefault<UBlueprintProfilerSettings>();
+	return Settings->GetPerformanceThresholdsModified() || InstrumentationEventQueue.Num() > 0;
 }
 
 void FBlueprintProfiler::ProcessEventProfilingData()
@@ -204,7 +223,7 @@ void FBlueprintProfiler::ProcessEventProfilingData()
 					if (InstanceEvent.IsNewInstance() && Event.IsEvent())
 					{
 						FEventRange NewEventRange;
-						NewEventRange.BlueprintContext = GetBlueprintContext(CurrEvent.GetBlueprintClassPath());
+						NewEventRange.BlueprintContext = FindBlueprintContext(CurrEvent.GetBlueprintClassPath());
 						if (NewEventRange.BlueprintContext.IsValid())
 						{
 							NewEventRange.InstanceName = NewEventRange.BlueprintContext->MapBlueprintInstance(InstanceEvent.GetInstancePath());
@@ -344,10 +363,10 @@ void FBlueprintProfiler::ProcessEventProfilingData()
 		}
 	}
 	// Update all active contexts if the display settings changed.
-	if (GetDefault<UBlueprintProfilerSettings>()->GetPerformanceThresholdsModified())
+	UBlueprintProfilerSettings* Settings = GetMutableDefault<UBlueprintProfilerSettings>();
+	if (Settings->GetPerformanceThresholdsModified())
 	{
-		GetMutableDefault<UBlueprintProfilerSettings>()->SetPerformanceThresholdsModified(false);
-
+		Settings->SetPerformanceThresholdsModified(false);
 		for (auto Context : PathToBlueprintContext)
 		{
 			DirtyContexts.Add(Context.Value);
@@ -365,8 +384,10 @@ void FBlueprintProfiler::RemoveAllBlueprintReferences(UBlueprint* Blueprint)
 	if (Blueprint)
 	{
 		const FString BlueprintClassPath = Blueprint->GeneratedClass->GetPathName();
-		if (PathToBlueprintContext.Contains(BlueprintClassPath))
+		if (TSharedPtr<FBlueprintExecutionContext>* ContextToRemove = PathToBlueprintContext.Find(BlueprintClassPath))
 		{
+			// Remove any registered contexts
+			RemoveUtilityContexts(*ContextToRemove);
 			// Tear down the blueprint context.
 			PathToBlueprintContext.Remove(BlueprintClassPath);
 			// Remove compilation hook, this will get added again if the blueprint is profiled
@@ -377,6 +398,22 @@ void FBlueprintProfiler::RemoveAllBlueprintReferences(UBlueprint* Blueprint)
 	}
 }
 
+void FBlueprintProfiler::AssociateUtilityContexts(TSharedPtr<FBlueprintExecutionContext> TargetContext)
+{
+	for (const FString UtilityContextPath : TargetContext->GetUtilityContexts())
+	{
+		PathToBlueprintUtilityContext.Add(UtilityContextPath) = TargetContext;
+	}
+}
+
+void FBlueprintProfiler::RemoveUtilityContexts(TSharedPtr<FBlueprintExecutionContext> TargetContext)
+{
+	for (const FString UtilityContextPath : TargetContext->GetUtilityContexts())
+	{
+		PathToBlueprintUtilityContext.Remove(UtilityContextPath);
+	}
+}
+
 #endif // WITH_EDITOR
 
 void FBlueprintProfiler::ResetProfilingData()
@@ -385,6 +422,7 @@ void FBlueprintProfiler::ResetProfilingData()
 	CaptureContext.ResetContext();
 #if WITH_EDITOR
 	PathToBlueprintContext.Reset();
+	PathToBlueprintUtilityContext.Reset();
 #endif // WITH_EDITOR
 }
 
diff --git a/Engine/Source/Developer/BlueprintProfiler/Private/BlueprintProfilerConnectionDrawingPolicy.cpp b/Engine/Source/Developer/BlueprintProfiler/Private/BlueprintProfilerConnectionDrawingPolicy.cpp
--- a/Engine/Source/Developer/BlueprintProfiler/Private/BlueprintProfilerConnectionDrawingPolicy.cpp
+++ b/Engine/Source/Developer/BlueprintProfiler/Private/BlueprintProfilerConnectionDrawingPolicy.cpp
@@ -13,18 +13,23 @@ FConnectionDrawingPolicy* FBlueprintProfilerPinConnectionFactory::CreateConnecti
 	FConnectionDrawingPolicy* NewPolicy = nullptr;
 	if (InGraphObj)
 	{
-		UBlueprint* Blueprint = InGraphObj->GetTypedOuter<UBlueprint>();
-		UBlueprintGeneratedClass* BPGC = Blueprint ? Cast<UBlueprintGeneratedClass>(Blueprint->GeneratedClass) : nullptr;
-		if (BPGC && BPGC->HasInstrumentation())
+		EBlueprintProfilerHeatMapDisplayMode WireHeatMode = GetDefault<UBlueprintProfilerSettings>()->WireHeatMapDisplayMode;
+		if (WireHeatMode != EBlueprintProfilerHeatMapDisplayMode::None)
 		{
-			EBlueprintProfilerHeatMapDisplayMode WireHeatMode = GetDefault<UBlueprintProfilerSettings>()->WireHeatMapDisplayMode;
-			if (WireHeatMode != EBlueprintProfilerHeatMapDisplayMode::None)
+			UBlueprint* Blueprint = InGraphObj->GetTypedOuter<UBlueprint>();
+			UBlueprintGeneratedClass* BPGC = Blueprint ? Cast<UBlueprintGeneratedClass>(Blueprint->GeneratedClass) : nullptr;
+			if (BPGC)
 			{
 				IBlueprintProfilerInterface& ProfilerModule = FModuleManager::LoadModuleChecked<IBlueprintProfilerInterface>("BlueprintProfiler");
-				TSharedPtr<FBlueprintExecutionContext> BlueprintContext = ProfilerModule.GetBlueprintContext(BPGC->GetPathName());
+				TSharedPtr<FBlueprintExecutionContext> BlueprintContext = ProfilerModule.FindBlueprintContext(BPGC->GetPathName());
 				if (BlueprintContext.IsValid())
 				{
-					NewPolicy = new FBlueprintProfilerConnectionDrawingPolicy(InBackLayerID, InFrontLayerID, ZoomFactor, InClippingRect, InDrawElements, WireHeatMode, InGraphObj);
+					NewPolicy = new FBlueprintProfilerConnectionDrawingPolicy(InBackLayerID, InFrontLayerID, ZoomFactor, InClippingRect, InDrawElements, WireHeatMode, InGraphObj, BlueprintContext);
+					// When we have an active connection drawing policy, associate utility contexts with the active blueprint execution context.
+					if (Blueprint->BlueprintType == BPTYPE_Normal)
+					{
+						ProfilerModule.AssociateUtilityContexts(BlueprintContext);
+					}
 				}
 			}
 		}
@@ -35,11 +40,15 @@ FConnectionDrawingPolicy* FBlueprintProfilerPinConnectionFactory::CreateConnecti
 /////////////////////////////////////////////////////
 // FBlueprintProfilerConnectionDrawingPolicy
 
-FBlueprintProfilerConnectionDrawingPolicy::FBlueprintProfilerConnectionDrawingPolicy(int32 InBackLayerID, int32 InFrontLayerID, float ZoomFactor, const FSlateRect& InClippingRect, FSlateWindowElementList& InDrawElements, EBlueprintProfilerHeatMapDisplayMode InHeatmapType, UEdGraph* InGraphObj)
+FBlueprintProfilerConnectionDrawingPolicy::FBlueprintProfilerConnectionDrawingPolicy(int32 InBackLayerID, int32 InFrontLayerID, float ZoomFactor, const FSlateRect& InClippingRect, FSlateWindowElementList& InDrawElements, EBlueprintProfilerHeatMapDisplayMode InHeatmapType, UEdGraph* InGraphObj, TSharedPtr<FBlueprintExecutionContext> InBlueprintContext)
 	: FKismetConnectionDrawingPolicy(InBackLayerID, InFrontLayerID, ZoomFactor, InClippingRect, InDrawElements, InGraphObj)
 	, WireHeatMode(InHeatmapType)
 	, GraphReference(InGraphObj)
+	, BlueprintContext(InBlueprintContext)
 {
+	AttackWireThickness = 8.f;
+	SustainWireThickness = 6.f;
+	ReleaseWireThickness = 4.f;
 }
 
 void FBlueprintProfilerConnectionDrawingPolicy::BuildExecutionRoadmap()
@@ -52,54 +61,42 @@ void FBlueprintProfilerConnectionDrawingPolicy::BuildExecutionRoadmap()
 		return;
 	}
 
-	UBlueprint* TargetBP = FBlueprintEditorUtils::FindBlueprintForGraphChecked(GraphObj);
-	UObject* ActiveObject = TargetBP->GetObjectBeingDebugged();
-	check(ActiveObject); // Due to CanBuildRoadmap
-
 	// Grab relevant profiler interfaces
-	UBlueprintGeneratedClass* TargetClass = Cast<UBlueprintGeneratedClass>(TargetBP->GeneratedClass);
-	IBlueprintProfilerInterface& ProfilerModule = FModuleManager::LoadModuleChecked<IBlueprintProfilerInterface>("BlueprintProfiler");
-	TSharedPtr<FBlueprintExecutionContext> BlueprintContext = ProfilerModule.GetBlueprintContext(TargetClass->GetPathName());
-	const FName InstanceName = BlueprintContext->MapBlueprintInstance(ActiveObject->GetPathName());
+	const FName InstanceName = BlueprintContext->GetActiveInstanceName();
 	TSharedPtr<FBlueprintFunctionContext> FunctionContext = BlueprintContext->GetFunctionContextFromGraph(GraphObj);
 
 	if (!FunctionContext.IsValid())
 	{
-		// No point proceeding without a valid function context.
+		// No point in proceeding without a valid function context.
 		return;
 	}
 
-	// Redirect the target Blueprint when debugging with a macro graph visible
-	if (TargetBP->BlueprintType == BPTYPE_MacroLibrary)
-	{
-		TargetBP = Cast<UBlueprint>(ActiveObject->GetClass()->ClassGeneratedBy);
-	}
-
 	TArray<double> SequentialNodeTimes;
 	TArray<UEdGraphPin*> SequentialExecPinsInGraph;
 	TArray<FTracePath> SequentialTracePathsInGraph;
 	TArray<TWeakPtr<FScriptExecutionNode>> SequentialProfilerNodesInGraph;
 	const FName ScopedFunctionName = FunctionContext->GetFunctionName();
+	const FName GraphName = FunctionContext->GetGraphName();
 
 	const TSimpleRingBuffer<FBlueprintExecutionTrace>& TraceHistory = BlueprintContext->GetTraceHistory();
 	for (int32 i = 0; i < TraceHistory.Num(); ++i)
 	{
 		const FBlueprintExecutionTrace& Sample = TraceHistory(i);
-		if (InstanceName == Sample.InstanceName && ScopedFunctionName == Sample.FunctionName)
+		const bool bInstanceMatch = InstanceName == SPDN_Blueprint ? true : (InstanceName == Sample.InstanceName);
+		if (bInstanceMatch && GraphName == Sample.GraphName)
 		{
 			if (Sample.ProfilerNode.IsValid())
 			{
-				const UEdGraphPin* Pin = FunctionContext->GetPinFromCodeLocation(Sample.Offset);
+				const bool bSamplePinValid = Sample.PinReference.Get() != nullptr;
+				// Patch in any missing pins using the script offset.
+				UEdGraphPin* Pin = bSamplePinValid ? Sample.PinReference.Get() : const_cast<UEdGraphPin*>(FunctionContext->GetPinFromCodeLocation(Sample.Offset));
 
 				if (Pin && Pin->PinType.PinCategory == UEdGraphSchema_K2::PC_Exec)
 				{
-					if (const UEdGraphNode* Node = Pin->GetOwningNode())
-					{
-						SequentialProfilerNodesInGraph.Add(Sample.ProfilerNode);
-						SequentialNodeTimes.Add(Sample.ObservationTime);
-						SequentialExecPinsInGraph.Add(const_cast<UEdGraphPin*>(Pin));
-						SequentialTracePathsInGraph.Add(Sample.TracePath);
-					}
+					SequentialProfilerNodesInGraph.Add(Sample.ProfilerNode);
+					SequentialNodeTimes.Add(Sample.ObservationTime);
+					SequentialExecPinsInGraph.Add(Pin);
+					SequentialTracePathsInGraph.Add(Sample.TracePath);
 				}
 			}
 		}
@@ -170,13 +167,13 @@ void FBlueprintProfilerConnectionDrawingPolicy::BuildExecutionRoadmap()
 						switch(WireHeatMode)
 						{
 							case EBlueprintProfilerHeatMapDisplayMode::Average:
+							case EBlueprintProfilerHeatMapDisplayMode::PinToPin:
 							{
 								ProfilerData.PredPerfData = CurPerfData->GetAverageHeatLevel();
 								ProfilerData.ThisPerfData = NextPerfData->GetAverageHeatLevel();
 								break;
 							}
 							case EBlueprintProfilerHeatMapDisplayMode::Inclusive:
-							case EBlueprintProfilerHeatMapDisplayMode::PinToPin:
 							{
 								ProfilerData.PredPerfData = CurPerfData->GetInclusiveHeatLevel();
 								ProfilerData.ThisPerfData = NextPerfData->GetInclusiveHeatLevel();
@@ -276,6 +273,11 @@ void FBlueprintProfilerConnectionDrawingPolicy::Draw(TMap<TSharedRef<SWidget>, F
 	}
 }
 
+bool FBlueprintProfilerConnectionDrawingPolicy::CanBuildRoadmap() const
+{
+	return BlueprintContext.IsValid() && BlueprintContext->IsContextValid();
+}
+
 void FBlueprintProfilerConnectionDrawingPolicy::DrawInterpColorSpline(const FGeometry& StartGeom, const FGeometry& EndGeom, const FScriptPerfConnectionParams& Params)
 {
 	//@TODO: These values should be pushed into the Slate style, they are compensating for a bit of
@@ -377,7 +379,7 @@ void FBlueprintProfilerConnectionDrawingPolicy::DrawPerfConnection(int32 LayerId
 	}
 
 	// Draw the spline itself
-	const float WireThickness = Params.WireThickness * 0.6f * ZoomFactor;
+	const float WireThickness = Params.WireThickness * ZoomFactor;
 	TArray<FSlateGradientStop> Gradients;
 	Gradients.Add(FSlateGradientStop(FVector2D::ZeroVector, Params.WireColor));
 	Gradients.Add(FSlateGradientStop(FVector2D::ZeroVector, Params.WireColor2));
@@ -408,7 +410,7 @@ void FBlueprintProfilerConnectionDrawingPolicy::DrawPerfConnection(int32 LayerId
 			const float BubbleOffset = FMath::Fmod(Time * BubbleSpeed, BubbleSpacing);
 			const int32 NumBubbles = FMath::CeilToInt(SplineLength/BubbleSpacing);
 			const float SizeMin = WireThickness * 0.05f;
-			const float SizeScale = WireThickness * 0.16f;
+			const float SizeScale = WireThickness * 0.10f;
 			const float SizeA = SizeMin + (Params.PerformanceData1 * SizeScale);
 			const float SizeB = SizeMin + (Params.PerformanceData2 * SizeScale);
 
@@ -467,4 +469,19 @@ void FBlueprintProfilerConnectionDrawingPolicy::DrawPerfConnection(int32 LayerId
 				);
 		}
 	}
+}
+
+void FBlueprintProfilerConnectionDrawingPolicy::ApplyHoverDeemphasis(UEdGraphPin* OutputPin, UEdGraphPin* InputPin, /*inout*/ float& Thickness, /*inout*/ FLinearColor& WireColor)
+{
+	const float FadeInBias = 0.75f; // Time in seconds before the fading starts to occur
+	const float FadeInPeriod = 0.6f; // Time in seconds after the bias before the fade is fully complete
+	const float TimeFraction = FMath::SmoothStep(0.0f, FadeInPeriod, (float)(FSlateApplication::Get().GetCurrentTime() - LastHoverTimeEvent - FadeInBias));
+
+	const bool bContainsBoth = HoveredPins.Contains(InputPin) && HoveredPins.Contains(OutputPin);
+	const bool bContainsOutput = HoveredPins.Contains(OutputPin);
+	const bool bEmphasize = bContainsBoth || (bContainsOutput && (InputPin == nullptr));
+	if (!bEmphasize)
+	{
+		Thickness = FMath::Lerp(Thickness, Thickness * 0.3f, TimeFraction);
+	}
 }
\ No newline at end of file
diff --git a/Engine/Source/Developer/BlueprintProfiler/Private/BlueprintProfilerConnectionDrawingPolicy.h b/Engine/Source/Developer/BlueprintProfiler/Private/BlueprintProfilerConnectionDrawingPolicy.h
--- a/Engine/Source/Developer/BlueprintProfiler/Private/BlueprintProfilerConnectionDrawingPolicy.h
+++ b/Engine/Source/Developer/BlueprintProfiler/Private/BlueprintProfilerConnectionDrawingPolicy.h
@@ -64,14 +64,16 @@ protected:
 
 public:
 
-	FBlueprintProfilerConnectionDrawingPolicy(int32 InBackLayerID, int32 InFrontLayerID, float ZoomFactor, const FSlateRect& InClippingRect, FSlateWindowElementList& InDrawElements, EBlueprintProfilerHeatMapDisplayMode InHeatmapType, UEdGraph* InGraphObj);
+	FBlueprintProfilerConnectionDrawingPolicy(int32 InBackLayerID, int32 InFrontLayerID, float ZoomFactor, const FSlateRect& InClippingRect, FSlateWindowElementList& InDrawElements, EBlueprintProfilerHeatMapDisplayMode InHeatmapType, UEdGraph* InGraphObj, TSharedPtr<class FBlueprintExecutionContext> InBlueprintContext);
 
 	// FKismetConnectionDrawingPolicy
+	virtual bool CanBuildRoadmap() const override;
 	virtual void BuildExecutionRoadmap() override;
 	// ~FKismetConnectionDrawingPolicy
 
 	// FConnectionDrawingPolicy interface
 	virtual void Draw(TMap<TSharedRef<SWidget>, FArrangedWidget>& PinGeometries, FArrangedChildren& ArrangedNodes) override;
+	virtual void ApplyHoverDeemphasis(UEdGraphPin* OutputPin, UEdGraphPin* InputPin, /*inout*/ float& Thickness, /*inout*/ FLinearColor& WireColor);
 	// End of FConnectionDrawingPolicy interface
 
 private:
@@ -91,5 +93,7 @@ private:
 	EBlueprintProfilerHeatMapDisplayMode WireHeatMode;
 	/** The current graph */
 	TWeakObjectPtr<UEdGraph> GraphReference;
+	/** The blueprint profiler context */
+	TSharedPtr<class FBlueprintExecutionContext> BlueprintContext;
 
 };
diff --git a/Engine/Source/Developer/BlueprintProfiler/Private/ScriptInstrumentationCapture.cpp b/Engine/Source/Developer/BlueprintProfiler/Private/ScriptInstrumentationCapture.cpp
--- a/Engine/Source/Developer/BlueprintProfiler/Private/ScriptInstrumentationCapture.cpp
+++ b/Engine/Source/Developer/BlueprintProfiler/Private/ScriptInstrumentationCapture.cpp
@@ -21,8 +21,22 @@ FName FScriptInstrumentedEvent::GetScopedFunctionName() const
 //////////////////////////////////////////////////////////////////////////
 // FInstrumentationCaptureContext
 
-void FInstrumentationCaptureContext::UpdateContext(const FScriptInstrumentationSignal& InstrumentSignal, TArray<FScriptInstrumentedEvent>& InstrumentationQueue)
+bool FInstrumentationCaptureContext::UpdateContext(const FScriptInstrumentationSignal& InstrumentSignal, TArray<FScriptInstrumentedEvent>& InstrumentationQueue)
 {
+	if (InstrumentSignal.GetType() == EScriptInstrumentation::InlineEvent)
+	{
+		// This inline event needs to be processed using the next incoming event's function class scope,
+		// so we need to store the data until the following UpdateContext call
+		PendingInlineEventName = InstrumentSignal.GetFunctionName();
+		return false;
+	}
+	else if (PendingInlineEventName != NAME_None)
+	{
+		// The pending inline event needs to be added first to make sure our events are processed in the proper order.
+		InstrumentationQueue.Add(FScriptInstrumentedEvent(EScriptInstrumentation::Event, InstrumentSignal.GetFunctionClassScope()->GetFName(), PendingInlineEventName));
+		PendingInlineEventName = NAME_None;
+	}
+
 	// Handle instance context switch
 	if (const UObject* NewContextObject = InstrumentSignal.GetContextObject())
 	{
@@ -43,14 +57,33 @@ void FInstrumentationCaptureContext::UpdateContext(const FScriptInstrumentationS
 		}
 		else
 		{
+			TWeakObjectPtr<const UClass> SignalScope = InstrumentSignal.GetFunctionClassScope();
 			// Handle function class scope changes
-			if (ContextFunctionClassScope != InstrumentSignal.GetFunctionClassScope())
+			if (ContextFunctionClassScope != SignalScope)
 			{
-				ContextFunctionClassScope = InstrumentSignal.GetFunctionClassScope();
-				InstrumentationQueue.Add(FScriptInstrumentedEvent(EScriptInstrumentation::ClassScope, ContextFunctionClassScope->GetFName(), InstrumentSignal.GetFunctionName()));
+				FName ClassScopeName = SignalScope->GetFName();
+				if (InstrumentSignal.GetType() == EScriptInstrumentation::Stop)
+				{
+					// In a Stop event, we need to use the last class scope on the stack to signify
+					// that we are reverting back to the previous class scope
+					ClassScopeName = ScopeStack.Pop();
+				}
+				else
+				{
+					// Stop events cause the Context to be reset so we need to check against the top scope on the stack before pushing it to the stack
+					if (ScopeStack.Num() == 0 || ScopeStack.Last() != ClassScopeName)
+					{
+						ScopeStack.Push(ClassScopeName);
+					}
+					ContextFunctionClassScope = SignalScope;
+				}
+
+				InstrumentationQueue.Add(FScriptInstrumentedEvent(EScriptInstrumentation::ClassScope, ClassScopeName, InstrumentSignal.GetFunctionName()));
 			}
 		}
 	}
+	
+	return true;
 }
 
 void FInstrumentationCaptureContext::ResetContext()
diff --git a/Engine/Source/Developer/BlueprintProfiler/Private/ScriptInstrumentationPlayback.cpp b/Engine/Source/Developer/BlueprintProfiler/Private/ScriptInstrumentationPlayback.cpp
--- a/Engine/Source/Developer/BlueprintProfiler/Private/ScriptInstrumentationPlayback.cpp
+++ b/Engine/Source/Developer/BlueprintProfiler/Private/ScriptInstrumentationPlayback.cpp
@@ -5,6 +5,7 @@
 #include "GraphEditorSettings.h"
 #include "BlueprintEditorUtils.h"
 #include "DelayAction.h"
+#include "K2Node_ForEachElementInEnum.h"
 
 #define LOCTEXT_NAMESPACE "ScriptInstrumentationPlayback"
 
@@ -152,6 +153,16 @@ FName FBlueprintExecutionContext::RemapInstancePath(const FName InstanceName) co
 	return Result;
 }
 
+FName FBlueprintExecutionContext::GetActiveInstanceName() const
+{
+	FName InstanceName = NAME_None;
+	if (BlueprintNode.IsValid())
+	{
+		InstanceName = BlueprintNode->GetActiveInstanceName();
+	}
+	return InstanceName;
+}
+
 bool FBlueprintExecutionContext::ResolveInstance(FName& InstanceNameInOut, TWeakObjectPtr<const UObject>& ObjectInOut)
 {
 	bool bNewInstance = false;
@@ -358,7 +369,21 @@ void FBlueprintExecutionContext::UpdateConnectedStats()
 	SCOPE_CYCLE_COUNTER(STAT_StatUpdateCost);
 	if (BlueprintNode.IsValid())
 	{
+		const UBlueprintProfilerSettings* ProfilerSettings = GetDefault<UBlueprintProfilerSettings>();
+		FName InstanceName = SPDN_Blueprint;
+		// Find active debug instance name.
+		if (ProfilerSettings->bDisplayByInstance && ProfilerSettings->bScopeToDebugInstance)
+		{
+			if (UBlueprint* OwnerBlueprint = Blueprint.Get())
+			{
+				if (UObject* DebugInstance = OwnerBlueprint->GetObjectBeingDebugged())
+				{
+					InstanceName = RemapInstancePath(FName(*DebugInstance->GetPathName()));
+				}
+			}
+		}
 		FTracePath InitialTracePath;
+		BlueprintNode->SetActiveInstanceName(InstanceName);
 		BlueprintNode->RefreshStats(InitialTracePath);
 	}
 }
@@ -404,11 +429,17 @@ bool FBlueprintExecutionContext::MapBlueprintExecution()
 			NewFunctionContext->DiscoverTunnels(GraphEntry.Value, DiscoveredTunnels);
 		}
 		// Create and map each tunnels instance as its own function context.
-		TSet<UEdGraph*> MappedTunnelGraphs;
 		for (auto TunnelInstance : DiscoveredTunnels)
 		{
 			if (UEdGraph* TunnelGraph = FBlueprintFunctionContext::GetGraphFromNode(TunnelInstance.Key, false))
 			{
+				// Register dependent context
+				UBlueprint* TunnelBlueprint = TunnelGraph->GetTypedOuter<UBlueprint>();
+				if (TunnelBlueprint->BlueprintType == BPTYPE_MacroLibrary)
+				{
+					const FString ClassPath = TunnelBlueprint->GeneratedClass->GetPathName();
+					DependentUtilityContexts.AddUnique(ClassPath);
+				}
 				// Create the tunnel instance context.
 				const FName TunnelInstanceFunctionName = FBlueprintFunctionContext::GetTunnelInstanceFunctionName(TunnelInstance.Key);
 				TSharedPtr<FBlueprintTunnelInstanceContext> NewTunnelInstanceContext = CreateFunctionContext<FBlueprintTunnelInstanceContext>(TunnelInstanceFunctionName, TunnelGraph);
@@ -511,32 +542,35 @@ void FBlueprintFunctionContext::InitialiseContextFromGraph(TSharedPtr<FBlueprint
 				}
 				for (auto EventNode : GraphEventNodes)
 				{
-					const FName EventName = GetScopedEventName(EventNode->GetFunctionName());
-					FScriptExecNodeParams EventParams;
-					EventParams.SampleFrequency = 1;
-					EventParams.NodeName = EventName;
-					EventParams.ObservedObject = EventNode;
-					if (bIsInheritedContext)
-					{
-						EventParams.Tooltip = LOCTEXT("NavigateToInheritedEventLocationHyperlink_ToolTip", "Navigate to the Inherited Event");
-						EventParams.IconColor = GetDefault<UGraphEditorSettings>()->ParentFunctionCallNodeTitleColor;
-						EventParams.NodeFlags = EScriptExecutionNodeFlags::Event|EScriptExecutionNodeFlags::InheritedEvent;
-						EventParams.DisplayName = FText::FromString(FString::Printf(TEXT("%s (%s)"), *EventNode->GetNodeTitle(ENodeTitleType::ListView).ToString(), *OwningBlueprint->GetName()));
-					}
-					else
+					if (EventNode->IsNodeEnabled())
 					{
-						EventParams.Tooltip = LOCTEXT("NavigateToEventLocationHyperlink_ToolTip", "Navigate to the Event");
-						EventParams.IconColor = GetDefault<UGraphEditorSettings>()->EventNodeTitleColor;
-						EventParams.NodeFlags = EScriptExecutionNodeFlags::Event;
-						EventParams.DisplayName = EventNode->GetNodeTitle(ENodeTitleType::ListView);
+						const FName EventName = GetScopedEventName(EventNode->GetFunctionName());
+						FScriptExecNodeParams EventParams;
+						EventParams.SampleFrequency = 1;
+						EventParams.NodeName = EventName;
+						EventParams.ObservedObject = EventNode;
+						if (bIsInheritedContext)
+						{
+							EventParams.Tooltip = LOCTEXT("NavigateToInheritedEventLocationHyperlink_ToolTip", "Navigate to the Inherited Event");
+							EventParams.IconColor = GetDefault<UGraphEditorSettings>()->ParentFunctionCallNodeTitleColor;
+							EventParams.NodeFlags = EScriptExecutionNodeFlags::Event|EScriptExecutionNodeFlags::RuntimeEvent|EScriptExecutionNodeFlags::InheritedEvent;
+							EventParams.DisplayName = FText::FromString(FString::Printf(TEXT("%s (%s)"), *EventNode->GetNodeTitle(ENodeTitleType::ListView).ToString(), *OwningBlueprint->GetName()));
+						}
+						else
+						{
+							EventParams.Tooltip = LOCTEXT("NavigateToEventLocationHyperlink_ToolTip", "Navigate to the Event");
+							EventParams.IconColor = GetDefault<UGraphEditorSettings>()->EventNodeTitleColor;
+							EventParams.NodeFlags = EScriptExecutionNodeFlags::Event|EScriptExecutionNodeFlags::RuntimeEvent;
+							EventParams.DisplayName = EventNode->GetNodeTitle(ENodeTitleType::ListView);
+						}
+						GetNodeCustomizations(EventParams);
+						const FSlateBrush* EventIcon = EventNode->ShowPaletteIconOnNode() ?	EventNode->GetIconAndTint(EventParams.IconColor).GetOptionalIcon() :
+																							FEditorStyle::GetBrush(TEXT("BlueprintProfiler.BPNode"));
+						EventParams.Icon = const_cast<FSlateBrush*>(EventIcon);
+						TSharedPtr<FScriptExecutionNode> EventExecNode = CreateExecutionNode(EventParams);
+						AddEntryPoint(EventExecNode);
+						BlueprintContextIn->AddEventNode(AsShared(), EventExecNode);
 					}
-					GetNodeCustomizations(EventParams);
-					const FSlateBrush* EventIcon = EventNode->ShowPaletteIconOnNode() ?	EventNode->GetIconAndTint(EventParams.IconColor).GetOptionalIcon() :
-																						FEditorStyle::GetBrush(TEXT("BlueprintProfiler.BPNode"));
-					EventParams.Icon = const_cast<FSlateBrush*>(EventIcon);
-					TSharedPtr<FScriptExecutionNode> EventExecNode = CreateExecutionNode(EventParams);
-					AddEntryPoint(EventExecNode);
-					BlueprintContextIn->AddEventNode(AsShared(), EventExecNode);
 				}
 				// Create any compiler generated events
 				FBlueprintDebugData& DebugData = BlueprintClass.Get()->GetDebugData();
@@ -550,34 +584,40 @@ void FBlueprintFunctionContext::InitialiseContextFromGraph(TSharedPtr<FBlueprint
 				// Create function node
 				for (auto FunctionEntry : FunctionEntryNodes)
 				{
-					FScriptExecNodeParams FunctionNodeParams;
-					FunctionNodeParams.SampleFrequency = 1;
-					FunctionNodeParams.NodeName = FunctionName;
-					FunctionNodeParams.ObservedObject = FunctionEntry;
-					if (bIsInheritedContext)
-					{
-						FunctionNodeParams.Tooltip = LOCTEXT("NavigateToInheritedFunctionLocationHyperlink_ToolTip", "Navigate to the Inherited Function");
-						FunctionNodeParams.IconColor = GetDefault<UGraphEditorSettings>()->ParentFunctionCallNodeTitleColor;
-						FunctionNodeParams.NodeFlags = EScriptExecutionNodeFlags::Event|EScriptExecutionNodeFlags::InheritedEvent;
-						FunctionNodeParams.DisplayName = FText::FromString(FString::Printf(TEXT("%s (%s)"), *GraphName.ToString(), *OwningBlueprint->GetName()));
-					}
-					else
-					{
-						FunctionNodeParams.Tooltip = LOCTEXT("NavigateToFunctionLocationHyperlink_ToolTip", "Navigate to the Function");
-						FunctionNodeParams.IconColor = GetDefault<UGraphEditorSettings>()->EventNodeTitleColor;
-						FunctionNodeParams.NodeFlags = EScriptExecutionNodeFlags::Event;
-						FunctionNodeParams.DisplayName = FText::FromName(GraphName);
-					}
-					GetNodeCustomizations(FunctionNodeParams);
-					const FSlateBrush* Icon = FunctionEntry->ShowPaletteIconOnNode() ? FunctionEntry->GetIconAndTint(FunctionNodeParams.IconColor).GetOptionalIcon() :
-																						FEditorStyle::GetBrush(TEXT("BlueprintProfiler.BPNode"));
-					FunctionNodeParams.Icon = const_cast<FSlateBrush*>(Icon);
-					TSharedPtr<FScriptExecutionNode> FunctionEntryNode = CreateExecutionNode(FunctionNodeParams);
-					AddEntryPoint(FunctionEntryNode);
-					// Add user construction scripts as events
-					if (GraphName == UEdGraphSchema_K2::FN_UserConstructionScript)
+					if (FunctionEntry->IsNodeEnabled())
 					{
-						BlueprintContextIn->AddEventNode(AsShared(), FunctionEntryNode);
+						FScriptExecNodeParams FunctionNodeParams;
+						FunctionNodeParams.SampleFrequency = 1;
+						FunctionNodeParams.NodeName = FunctionName;
+						FunctionNodeParams.ObservedObject = FunctionEntry;
+						const bool bConstructionScriptEntryPoint = GraphName == UEdGraphSchema_K2::FN_UserConstructionScript;
+						if (bIsInheritedContext)
+						{
+							FunctionNodeParams.Tooltip = LOCTEXT("NavigateToInheritedFunctionLocationHyperlink_ToolTip", "Navigate to the Inherited Function");
+							FunctionNodeParams.IconColor = GetDefault<UGraphEditorSettings>()->ParentFunctionCallNodeTitleColor;
+							FunctionNodeParams.NodeFlags = EScriptExecutionNodeFlags::Event|EScriptExecutionNodeFlags::InheritedEvent;
+							FunctionNodeParams.NodeFlags |= bConstructionScriptEntryPoint ? EScriptExecutionNodeFlags::ConstructionEvent : EScriptExecutionNodeFlags::RuntimeEvent;
+							FunctionNodeParams.DisplayName = FText::FromString(FString::Printf(TEXT("%s (%s)"), *GraphName.ToString(), *OwningBlueprint->GetName()));
+						}
+						else
+						{
+							FunctionNodeParams.Tooltip = LOCTEXT("NavigateToFunctionLocationHyperlink_ToolTip", "Navigate to the Function");
+							FunctionNodeParams.IconColor = GetDefault<UGraphEditorSettings>()->EventNodeTitleColor;
+							FunctionNodeParams.NodeFlags = EScriptExecutionNodeFlags::Event;
+							FunctionNodeParams.NodeFlags |= bConstructionScriptEntryPoint ? EScriptExecutionNodeFlags::ConstructionEvent : EScriptExecutionNodeFlags::RuntimeEvent;
+							FunctionNodeParams.DisplayName = FText::FromName(GraphName);
+						}
+						GetNodeCustomizations(FunctionNodeParams);
+						const FSlateBrush* Icon = FunctionEntry->ShowPaletteIconOnNode() ? FunctionEntry->GetIconAndTint(FunctionNodeParams.IconColor).GetOptionalIcon() :
+																							FEditorStyle::GetBrush(TEXT("BlueprintProfiler.BPNode"));
+						FunctionNodeParams.Icon = const_cast<FSlateBrush*>(Icon);
+						TSharedPtr<FScriptExecutionNode> FunctionEntryNode = CreateExecutionNode(FunctionNodeParams);
+						AddEntryPoint(FunctionEntryNode);
+						// Add user construction scripts as events
+						if (bConstructionScriptEntryPoint)
+						{
+							BlueprintContextIn->AddEventNode(AsShared(), FunctionEntryNode);
+						}
 					}
 				}
 			}
@@ -774,7 +814,8 @@ void FBlueprintFunctionContext::CreateDelegatePinEvents(TSharedPtr<FBlueprintExe
 				EventParams.OwningGraphName = NAME_None;
 				EventParams.DisplayName = NodeEvents.Key->GetNodeTitle(ENodeTitleType::ListView);
 				EventParams.Tooltip = LOCTEXT("NavigateToEventLocationHyperlink_ToolTip", "Navigate to the Event");
-				EventParams.NodeFlags = FunctionContext->IsInheritedContext() ? (EScriptExecutionNodeFlags::Event|EScriptExecutionNodeFlags::InheritedEvent) : EScriptExecutionNodeFlags::Event;
+				EventParams.NodeFlags = FunctionContext->IsInheritedContext() ? (EScriptExecutionNodeFlags::Event|EScriptExecutionNodeFlags::RuntimeEvent|EScriptExecutionNodeFlags::InheritedEvent) : 
+																				(EScriptExecutionNodeFlags::Event|EScriptExecutionNodeFlags::RuntimeEvent);
 				GetNodeCustomizations(EventParams);
 				const FSlateBrush* EventIcon = NodeEvents.Key->ShowPaletteIconOnNode() ? NodeEvents.Key->GetIconAndTint(EventParams.IconColor).GetOptionalIcon() :
 																						 FEditorStyle::GetBrush(TEXT("BlueprintProfiler.BPNode"));
@@ -924,6 +965,17 @@ void FBlueprintFunctionContext::DetermineGraphNodeCharacteristics(const UEdGraph
 			NodeParams.SampleFrequency = ConnectedExecPins;
 			NodeParams.NodeFlags |= EScriptExecutionNodeFlags::SequentialBranch;
 		}
+		else if (const UK2Node_ForEachElementInEnum* ForEachElementInEnumNode = Cast<UK2Node_ForEachElementInEnum>(GraphNode))
+		{
+			// This behaves like a ForLoop macro instance node, but w/o an actual macro expansion.
+			const UEdGraphPin* LoopBodyPin = ForEachElementInEnumNode->FindPin(UK2Node_ForEachElementInEnum::InsideLoopPinName);
+			if (LoopBodyPin && LoopBodyPin->LinkedTo.Num() > 0)
+			{
+				// Subtract 1, because we've already accounted for the first iteration pass through the link above.
+				NodeParams.SampleFrequency = ConnectedExecPins + ForEachElementInEnumNode->Enum->GetMaxEnumValue() - 1;
+			}
+			NodeParams.NodeFlags |= EScriptExecutionNodeFlags::SequentialBranch;
+		}
 		else
 		{
 			if (GraphNode->IsA<UK2Node_Event>() || GraphNode->IsA<UK2Node_InputKey>() || GraphNode->IsA<UK2Node_InputTouch>())
@@ -1257,7 +1309,7 @@ void FBlueprintFunctionContext::MapExecPins(TSharedPtr<FScriptExecutionNode> Exe
 	}
 }
 
-TSharedPtr<FScriptExecutionNode> FBlueprintFunctionContext::GetTunnelBoundaryNodeChecked(const UEdGraphPin* TunnelPin)
+TSharedPtr<FScriptExecutionNode> FBlueprintFunctionContext::GetTunnelBoundaryNode(const UEdGraphPin* TunnelPin)
 {
 	TSharedPtr<FScriptExecutionNode> Result;
 	if (const UK2Node_Tunnel* TunnelNode = Cast<UK2Node_Tunnel>(TunnelPin->GetOwningNode()))
@@ -1278,7 +1330,14 @@ TSharedPtr<FScriptExecutionNode> FBlueprintFunctionContext::GetTunnelBoundaryNod
 			Result = GetProfilerDataForNode(GetPinName(TunnelPin));
 		}
 	}
-	check (Result.IsValid());
+
+	return Result;
+}
+
+TSharedPtr<FScriptExecutionNode> FBlueprintFunctionContext::GetTunnelBoundaryNodeChecked(const UEdGraphPin* TunnelPin)
+{
+	TSharedPtr<FScriptExecutionNode> Result = GetTunnelBoundaryNode(TunnelPin);
+	check(Result.IsValid());
 	return Result;
 }
 
@@ -1287,8 +1346,8 @@ TSharedPtr<FScriptExecutionNode> FBlueprintFunctionContext::MapTunnelBoundary(co
 	TSharedPtr<FScriptExecutionNode> TunnelBoundaryNode;
 	if (TunnelPin)
 	{
-		TunnelBoundaryNode = GetTunnelBoundaryNodeChecked(TunnelPin);
-		if (TunnelBoundaryNode->IsTunnelEntry())
+		TunnelBoundaryNode = GetTunnelBoundaryNode(TunnelPin);
+		if (TunnelBoundaryNode.IsValid() && TunnelBoundaryNode->IsTunnelEntry())
 		{
 			TSharedPtr<FScriptExecutionTunnelEntry> TunnelEntryInstance = StaticCastSharedPtr<FScriptExecutionTunnelEntry>(TunnelBoundaryNode);
 			for (auto ExitSite : TunnelEntryInstance->GetLinkedNodes())
@@ -1451,8 +1510,16 @@ FName FBlueprintFunctionContext::GetTunnelInstanceFunctionName(const UEdGraphNod
 	UEdGraph* OuterGraph = GraphNode ? GraphNode->GetTypedOuter<UEdGraph>() : nullptr;
 	if (OuterGraph)
 	{
-		// To identify instances exactly, we need the owning graph name and the node name to avoid collisions.
-		ScopedFunctionName = FName(*FString::Printf(TEXT("%s::%s"), *OuterGraph->GetName(), *GraphNode->GetName()));
+		UBlueprint* OwnerBlueprint = FBlueprintEditorUtils::FindBlueprintForGraph(OuterGraph);
+		if (OwnerBlueprint)
+		{
+			UBlueprintGeneratedClass* BPGC = Cast<UBlueprintGeneratedClass>(OwnerBlueprint->GeneratedClass);
+			if (BPGC)
+			{
+				// To identify instances exactly, we need the owning blueprint name, owning graph name, and the node name to avoid collisions.
+				ScopedFunctionName = FName(*FString::Printf(TEXT("%s::%s::%s"), *BPGC->GetName(), *OuterGraph->GetName(), *GraphNode->GetName()));
+			}
+		}
 	}
 	check (ScopedFunctionName != NAME_None);
 	return ScopedFunctionName;
@@ -1850,7 +1917,7 @@ void FBlueprintTunnelInstanceContext::MapTunnelContext(TSharedPtr<FBlueprintFunc
 				// Locate the tunnel instance context.
 				const FName TunnelInstanceFunctionName = GetTunnelInstanceFunctionName(NestedTunnelInstance.Key);
 				const FName TunnelFunctionName = GetScopedFunctionNameFromNode(NestedTunnelInstance.Key);
-				TSharedPtr<FBlueprintFunctionContext> FunctionContext = BlueprintContext.Pin()->GetFunctionContext(TunnelFunctionName);
+				TSharedPtr<FBlueprintFunctionContext> FunctionContext = BlueprintContext.Pin()->GetFunctionContext(TunnelInstanceFunctionName);
 				if (!FunctionContext.IsValid())
 				{
 					TSharedPtr<FBlueprintTunnelInstanceContext> NestedContext = MakeShareable(new FBlueprintTunnelInstanceContext);
@@ -2251,7 +2318,7 @@ void FBlueprintTunnelInstanceContext::MapInputPins(TSharedPtr<FScriptExecutionNo
 	}
 }
 
-TSharedPtr<FScriptExecutionNode> FBlueprintTunnelInstanceContext::GetTunnelBoundaryNodeChecked(const UEdGraphPin* TunnelPin)
+TSharedPtr<FScriptExecutionNode> FBlueprintTunnelInstanceContext::GetTunnelBoundaryNode(const UEdGraphPin* TunnelPin)
 {
 	TSharedPtr<FScriptExecutionNode> Result;
 	// Check tunnel node
@@ -2259,15 +2326,16 @@ TSharedPtr<FScriptExecutionNode> FBlueprintTunnelInstanceContext::GetTunnelBound
 	{
 		if (FBlueprintEditorUtils::IsTunnelInstanceNode(TunnelNode))
 		{
+			const FName TunnelInstanceFunctionName = GetTunnelInstanceFunctionName(TunnelNode);
+
 			// Lookup internal tunnel boundary.
-			if (TWeakPtr<FBlueprintFunctionContext>* NestedTunnelContext = ChildFunctionContexts.Find(TunnelNode->GetFName()))
+			if (TWeakPtr<FBlueprintFunctionContext>* NestedTunnelContext = ChildFunctionContexts.Find(TunnelInstanceFunctionName))
 			{
 				Result = NestedTunnelContext->Pin()->GetProfilerDataForNode(GetTunnelBoundaryName(TunnelPin));
 			}
 			else
 			{
 				// Lookup external tunnel boundary.
-				const FName TunnelInstanceFunctionName = GetTunnelInstanceFunctionName(TunnelNode);
 				TSharedPtr<FBlueprintFunctionContext> TunnelContext = BlueprintContext.Pin()->GetFunctionContext(TunnelInstanceFunctionName);
 				if (TunnelContext.IsValid())
 				{
@@ -2291,7 +2359,7 @@ TSharedPtr<FScriptExecutionNode> FBlueprintTunnelInstanceContext::GetTunnelBound
 			}
 		}
 	}
-	check (Result.IsValid());
+
 	return Result;
 }
 
@@ -2439,7 +2507,7 @@ bool FScriptEventPlayback::Process(const TArray<FScriptInstrumentedEvent>& Signa
 			if (CurrSignal.GetType() == EScriptInstrumentation::Class)
 			{
 				// Update the current mapped blueprint context.
-				BlueprintContext = Profiler->GetBlueprintContext(CurrSignal.GetBlueprintClassPath());
+				BlueprintContext = Profiler->GetOrCreateBlueprintContext(CurrSignal.GetBlueprintClassPath());
 
 				// Skip to the next signal.
 				continue;
@@ -2605,7 +2673,8 @@ bool FScriptEventPlayback::Process(const TArray<FScriptInstrumentedEvent>& Signa
 						AddToTraceHistory(CurrentNodeData.ImpureNode, CurrSignal);
 						// Process node exit
 						const int32 ScriptCodeExit = CurrSignal.GetScriptCodeOffset();
-						if (const UEdGraphPin* ValidPin = FunctionContext->GetPinFromCodeLocation(ScriptCodeExit))
+						const UEdGraphPin* ValidPin = FunctionContext->GetPinFromCodeLocation(ScriptCodeExit);
+						if (ValidPin && ValidPin->GetOwningNode() == GraphNode)
 						{
 							// Delegate/event pin entry points require the tracepath to be reset.
 							TSharedPtr<FScriptExecutionNode> Pin = FunctionContext->GetProfilerDataForNode(FBlueprintFunctionContext::GetUniquePinName(ValidPin));
@@ -2662,7 +2731,7 @@ bool FScriptEventPlayback::Process(const TArray<FScriptInstrumentedEvent>& Signa
 			EventNode = FunctionContext->GetProfilerDataForNode(EventTiming.Key);
 			check (EventNode.IsValid());
 			TSharedPtr<FScriptPerfData> PerfData = EventNode->GetOrAddPerfDataByInstanceAndTracePath(InstanceName, EventTracePath);
-			PerfData->AddEventTiming(SignalData[StopIdx].GetTime() - SignalData[LastEventIdx].GetTime());
+			PerfData->AddEventTiming(EventTiming.Value);
 		}
 		EventTimings.Reset();
 		// Process Node map timings -- this can probably be rolled back into submission during the initial processing and lose this extra iteration.
@@ -2798,14 +2867,14 @@ void FScriptEventPlayback::ProcessTunnelBoundary(NodeSignalHelper& CurrentNodeDa
 		{
 			// Grab the tunnel entry
 			TSharedPtr<FScriptExecutionTunnelEntry> TunnelEntry = StaticCastSharedPtr<FScriptExecutionTunnelEntry>(TunnelBoundary);
+			// Add Tunnel Trace History
+			AddTunnelTraceHistory(TunnelBoundary, CurrSignal, TracePath, TracePath);
+			// Update the tunnel entry count so we can fixup stat samples later.
+			TunnelEntry->IncrementTunnelEntryCount(InstanceName, TracePath);
 			// Process tunnel entry sites
 			TunnelTraceStack.Push(TracePath);
 			TracePath = FTracePath(TracePath, TunnelEntry);
 			CurrentNodeData.AverageEvents.Add(CurrSignal);
-			// Update the tunnel entry count so we can fixup stat samples later.
-			TunnelEntry->IncrementTunnelEntryCount();
-			// Add Trace History
-			AddToTraceHistory(TunnelBoundary, CurrSignal);
 		}
 		else if (TunnelBoundary->HasFlags(EScriptExecutionNodeFlags::TunnelExit))
 		{
@@ -2813,8 +2882,12 @@ void FScriptEventPlayback::ProcessTunnelBoundary(NodeSignalHelper& CurrentNodeDa
 			if (CurrentNodeData.AverageEvents.Num())
 			{
 				TSharedPtr<FScriptExecutionTunnelEntry> Tunnel = TracePath.GetTunnel();
+				// Cache the tunnel tracepath and restore the non tunnel tracepath as current.
 				FTracePath InternalExitTrace = TracePath;
 				TracePath = TunnelTraceStack.Pop();
+				// Add Tunnel Trace History
+				AddTunnelTraceHistory(TunnelBoundary, CurrSignal, InternalExitTrace, TracePath);
+				// Add timings to tunnel nodes.
 				if (Tunnel.IsValid())
 				{
 					// Update the entry and both exit sites ( one inside the tunnel and one exit site )
@@ -2823,8 +2896,6 @@ void FScriptEventPlayback::ProcessTunnelBoundary(NodeSignalHelper& CurrentNodeDa
 				}
 				TracePath.AddExitPin(ScriptCodeOffset);
 				CurrentNodeData.AverageEvents.Reset();
-				// Add Trace History
-				AddToTraceHistory(TunnelBoundary, CurrSignal);
 			}
 		}
 	}
@@ -2885,15 +2956,107 @@ void FScriptEventPlayback::ProcessExecutionSequence(NodeSignalHelper& CurrentNod
 
 void FScriptEventPlayback::AddToTraceHistory(const TSharedPtr<FScriptExecutionNode> ProfilerNode, const FScriptInstrumentedEvent& TraceSignal)
 {
-	// Add Trace History
-	FBlueprintExecutionTrace& NewTraceEvent = BlueprintContext->AddNewTraceHistoryEvent();
-	NewTraceEvent.TraceType = TraceSignal.GetType();
-	NewTraceEvent.TracePath = TracePath;
-	NewTraceEvent.InstanceName = InstanceName;
-	NewTraceEvent.FunctionName = CurrentFunctionName;
-	NewTraceEvent.ProfilerNode = ProfilerNode;
-	NewTraceEvent.Offset = TraceSignal.GetScriptCodeOffset();
-	NewTraceEvent.ObservationTime = TraceSignal.GetTime();
+	bool bAddHistoryEvent = true;
+	if (ProfilerNode->IsBranch() && TraceSignal.IsNodeExit())
+	{
+		// Add exec pins into the trace history for correct heat displays.
+		TSharedPtr<FScriptExecutionNode> PinNode = ProfilerNode->GetLinkedNodeByScriptOffset(TraceSignal.GetScriptCodeOffset());
+		if (PinNode.IsValid() && PinNode->HasFlags(EScriptExecutionNodeFlags::ExecPin))
+		{
+			FBlueprintExecutionTrace& NewTraceEvent = BlueprintContext->AddNewTraceHistoryEvent();
+			NewTraceEvent.TraceType = TraceSignal.GetType();
+			NewTraceEvent.TracePath = TracePath;
+			NewTraceEvent.InstanceName = InstanceName;
+			NewTraceEvent.FunctionName = CurrentFunctionName;
+			NewTraceEvent.GraphName = PinNode->GetGraphName();
+			NewTraceEvent.Offset = TraceSignal.GetScriptCodeOffset();
+			NewTraceEvent.ObservationTime = TraceSignal.GetTime();
+			NewTraceEvent.ProfilerNode = PinNode;
+			NewTraceEvent.PinReference = PinNode->GetObservedPin();
+			bAddHistoryEvent = false;
+		}
+	}
+	if (bAddHistoryEvent)
+	{
+		// Add Trace History
+		FBlueprintExecutionTrace& NewTraceEvent = BlueprintContext->AddNewTraceHistoryEvent();
+		NewTraceEvent.TraceType = TraceSignal.GetType();
+		NewTraceEvent.TracePath = TracePath;
+		NewTraceEvent.InstanceName = InstanceName;
+		NewTraceEvent.FunctionName = CurrentFunctionName;
+		NewTraceEvent.GraphName = ProfilerNode->GetGraphName();
+		NewTraceEvent.Offset = TraceSignal.GetScriptCodeOffset();
+		NewTraceEvent.ObservationTime = TraceSignal.GetTime();
+		NewTraceEvent.ProfilerNode = ProfilerNode;
+		NewTraceEvent.PinReference = ProfilerNode->GetObservedPin();
+	}
+}
+
+void FScriptEventPlayback::AddTunnelTraceHistory(const TSharedPtr<FScriptExecutionNode> TunnelBoundary, const FScriptInstrumentedEvent& TraceSignal, const FTracePath& InternalPath, const FTracePath& ExternalPath)
+{
+	if (TunnelBoundary->HasFlags(EScriptExecutionNodeFlags::TunnelEntry))
+	{
+		// Add history for external tunnel entry pin
+		FBlueprintExecutionTrace& ExternalTraceEvent = BlueprintContext->AddNewTraceHistoryEvent();
+		ExternalTraceEvent.TraceType = TraceSignal.GetType();
+		ExternalTraceEvent.TracePath = ExternalPath;
+		ExternalTraceEvent.InstanceName = InstanceName;
+		ExternalTraceEvent.FunctionName = CurrentFunctionName;
+		ExternalTraceEvent.GraphName = TunnelBoundary->GetGraphName();
+		ExternalTraceEvent.Offset = TraceSignal.GetScriptCodeOffset();
+		ExternalTraceEvent.ObservationTime = TraceSignal.GetTime();
+		ExternalTraceEvent.ProfilerNode = TunnelBoundary;
+		ExternalTraceEvent.PinReference = TunnelBoundary->GetObservedPin();
+		// Add history for internal tunnel entry pin
+		if (TunnelBoundary->GetNumChildren() > 0)
+		{
+			// We might want to try harder to find the right child node here?!.
+			TSharedPtr<FScriptExecutionNode> InternalTunnelNode = TunnelBoundary->GetChildByIndex(0);
+			if (InternalTunnelNode.IsValid() && InternalTunnelNode->HasFlags(EScriptExecutionNodeFlags::TunnelEntryPin))
+			{
+				FBlueprintExecutionTrace& InternalTraceEvent = BlueprintContext->AddNewTraceHistoryEvent();
+				InternalTraceEvent.TraceType = TraceSignal.GetType();
+				InternalTraceEvent.TracePath = InternalPath;
+				InternalTraceEvent.InstanceName = InstanceName;
+				InternalTraceEvent.FunctionName = CurrentFunctionName;
+				InternalTraceEvent.GraphName = InternalTunnelNode->GetGraphName();
+				InternalTraceEvent.Offset = TraceSignal.GetScriptCodeOffset();
+				InternalTraceEvent.ObservationTime = TraceSignal.GetTime();
+				InternalTraceEvent.ProfilerNode = InternalTunnelNode;
+				InternalTraceEvent.PinReference = InternalTunnelNode->GetObservedPin();
+			}
+		}
+	}
+	else if (TunnelBoundary->HasFlags(EScriptExecutionNodeFlags::TunnelExit))
+	{
+		// Add history for internal tunnel exit pin
+		FBlueprintExecutionTrace& InternalTraceEvent = BlueprintContext->AddNewTraceHistoryEvent();
+		InternalTraceEvent.TraceType = TraceSignal.GetType();
+		InternalTraceEvent.TracePath = InternalPath;
+		InternalTraceEvent.InstanceName = InstanceName;
+		InternalTraceEvent.FunctionName = CurrentFunctionName;
+		InternalTraceEvent.GraphName = TunnelBoundary->GetGraphName();
+		InternalTraceEvent.Offset = TraceSignal.GetScriptCodeOffset();
+		InternalTraceEvent.ObservationTime = TraceSignal.GetTime();
+		InternalTraceEvent.ProfilerNode = TunnelBoundary;
+		InternalTraceEvent.PinReference = TunnelBoundary->GetObservedPin();
+		// Add history for external tunnel exit pin
+		TSharedPtr<FScriptExecutionTunnelEntry> TunnelEntryNode = InternalPath.GetTunnel();
+		TSharedPtr<FScriptExecutionTunnelExit> TunnelExit = StaticCastSharedPtr<FScriptExecutionTunnelExit>(TunnelBoundary);
+		if (TunnelExit.IsValid() && TunnelEntryNode.IsValid())
+		{
+			FBlueprintExecutionTrace& ExternalTraceEvent = BlueprintContext->AddNewTraceHistoryEvent();
+			ExternalTraceEvent.TraceType = TraceSignal.GetType();
+			ExternalTraceEvent.TracePath = ExternalPath;
+			ExternalTraceEvent.InstanceName = InstanceName;
+			ExternalTraceEvent.FunctionName = CurrentFunctionName;
+			ExternalTraceEvent.GraphName = TunnelEntryNode->GetGraphName();
+			ExternalTraceEvent.Offset = TraceSignal.GetScriptCodeOffset();
+			ExternalTraceEvent.ObservationTime = TraceSignal.GetTime();
+			ExternalTraceEvent.ProfilerNode = TunnelExit;
+			ExternalTraceEvent.PinReference = TunnelExit->GetExternalPin();
+		}
+	}
 }
 
 #undef LOCTEXT_NAMESPACE
diff --git a/Engine/Source/Developer/BlueprintProfiler/Public/BlueprintProfiler.h b/Engine/Source/Developer/BlueprintProfiler/Public/BlueprintProfiler.h
--- a/Engine/Source/Developer/BlueprintProfiler/Public/BlueprintProfiler.h
+++ b/Engine/Source/Developer/BlueprintProfiler/Public/BlueprintProfiler.h
@@ -33,11 +33,14 @@ public:
 #if WITH_EDITOR
 	virtual void AddInstrumentedBlueprint(UBlueprint* InstrumentedBlueprint) override;
 	virtual FOnBPStatGraphLayoutChanged& GetGraphLayoutChangedDelegate() { return GraphLayoutChangedDelegate; }
-	virtual TSharedPtr<FBlueprintExecutionContext> GetBlueprintContext(const FString& BlueprintClassPath) override;
+	virtual TSharedPtr<FBlueprintExecutionContext> GetOrCreateBlueprintContext(const FString& BlueprintClassPath) override;
+	virtual TSharedPtr<FBlueprintExecutionContext> FindBlueprintContext(const FString& BlueprintClassPath) override;
 	virtual TSharedPtr<FScriptExecutionNode> GetProfilerDataForNode(const UEdGraphNode* GraphNode) override;
 	virtual TSharedPtr<FScriptExecutionBlueprint> GetProfilerDataForBlueprint(const UBlueprint* Blueprint) override;
 	virtual bool HasDataForInstance(const UObject* Instance) const override;
 	virtual void ProcessEventProfilingData() override;
+	virtual void AssociateUtilityContexts(TSharedPtr<FBlueprintExecutionContext> TargetContext);
+	virtual void RemoveUtilityContexts(TSharedPtr<FBlueprintExecutionContext> TargetContext);
 #endif
 	// End IBlueprintProfilerModule
 
@@ -84,8 +87,10 @@ protected:
 	bool bPIEActive;
 	/** Suspended script events, the key is the latent LinkId */
 	TMap<FName, TMap<int32, TSharedPtr<FScriptEventPlayback>>> SuspendedEvents;
-	/** Cached object path and code offset lookup to UObjects */
+	/** Cached object path lookup to blueprints */
 	TMap<FString, TSharedPtr<FBlueprintExecutionContext>> PathToBlueprintContext;
+	/** Cached object path lookup to utility blueprints */
+	TMap<FString, TSharedPtr<FBlueprintExecutionContext>> PathToBlueprintUtilityContext;
 	/** The profiler connection factory */
 	TSharedPtr<struct FGraphPanelPinConnectionFactory> ConnectionFactory;
 #endif
diff --git a/Engine/Source/Developer/BlueprintProfiler/Public/BlueprintProfilerModule.h b/Engine/Source/Developer/BlueprintProfiler/Public/BlueprintProfilerModule.h
--- a/Engine/Source/Developer/BlueprintProfiler/Public/BlueprintProfilerModule.h
+++ b/Engine/Source/Developer/BlueprintProfiler/Public/BlueprintProfilerModule.h
@@ -35,7 +35,10 @@ public:
 	virtual FOnBPStatGraphLayoutChanged& GetGraphLayoutChangedDelegate() = 0;
 
 	/** Returns the cached blueprint context or creates and maps a new entry */
-	virtual TSharedPtr<class FBlueprintExecutionContext> GetBlueprintContext(const FString& BlueprintClassPath) = 0;
+	virtual TSharedPtr<class FBlueprintExecutionContext> GetOrCreateBlueprintContext(const FString& BlueprintClassPath) = 0;
+
+	/** Returns the cached blueprint context if existing */
+	virtual TSharedPtr<class FBlueprintExecutionContext> FindBlueprintContext(const FString& BlueprintClassPath) = 0;
 
 	/** Returns the current profiling event data for node */
 	virtual TSharedPtr<class FScriptExecutionNode> GetProfilerDataForNode(const UEdGraphNode* GraphNode) = 0;
@@ -49,6 +52,12 @@ public:
 	/** Process profiling event data */
 	virtual void ProcessEventProfilingData() {}
 
+	/** Associate utility contexts with the argument blueprint context */
+	virtual void AssociateUtilityContexts(TSharedPtr<FBlueprintExecutionContext> TargetContext) {}
+
+	/** Remove associated utility contexts for the argument blueprint context */
+	virtual void RemoveUtilityContexts(TSharedPtr<FBlueprintExecutionContext> TargetContext) {}
+
 #endif // WITH_EDITOR
 
 };
\ No newline at end of file
diff --git a/Engine/Source/Developer/BlueprintProfiler/Public/ScriptInstrumentationCapture.h b/Engine/Source/Developer/BlueprintProfiler/Public/ScriptInstrumentationCapture.h
--- a/Engine/Source/Developer/BlueprintProfiler/Public/ScriptInstrumentationCapture.h
+++ b/Engine/Source/Developer/BlueprintProfiler/Public/ScriptInstrumentationCapture.h
@@ -109,7 +109,7 @@ class FInstrumentationCaptureContext
 public:
 
 	/** Update the current execution context, and emit events to signal changes */
-	void UpdateContext(const FScriptInstrumentationSignal& InstrumentSignal, TArray<FScriptInstrumentedEvent>& InstrumentationQueue);
+	bool UpdateContext(const FScriptInstrumentationSignal& InstrumentSignal, TArray<FScriptInstrumentedEvent>& InstrumentationQueue);
 
 	/** Reset the current event context */
 	void ResetContext();
@@ -124,4 +124,8 @@ private:
 	TWeakObjectPtr<const UObject> ContextObject;
 	/** The current event context */
 	FName ActiveEvent;
+	/** The current pending inline Event. NAME_None if no event is pending */
+	FName PendingInlineEventName;
+	/** The Class Scope Stack for traversing through parent blueprint calls */
+	TArray<FName> ScopeStack;
 };
diff --git a/Engine/Source/Developer/BlueprintProfiler/Public/ScriptInstrumentationPlayback.h b/Engine/Source/Developer/BlueprintProfiler/Public/ScriptInstrumentationPlayback.h
--- a/Engine/Source/Developer/BlueprintProfiler/Public/ScriptInstrumentationPlayback.h
+++ b/Engine/Source/Developer/BlueprintProfiler/Public/ScriptInstrumentationPlayback.h
@@ -13,7 +13,9 @@ struct FBlueprintExecutionTrace
 	FTracePath TracePath;
 	FName InstanceName;
 	FName FunctionName;
+	FName GraphName;
 	TWeakPtr<FScriptExecutionNode> ProfilerNode;
+	FEdGraphPinReference PinReference;
 	int32 Offset;
 	double ObservationTime;
 };
@@ -52,6 +54,9 @@ public:
 	/** Returns the blueprint exec node for this context */
 	TSharedPtr<class FScriptExecutionBlueprint> GetBlueprintExecNode() const { return BlueprintNode; }
 
+	/** Returns an array of class paths that represent dependent blueprint contexts for this blueprint context */
+	const TArray<FString>& GetUtilityContexts() const { return DependentUtilityContexts; }
+
 	/** Returns the blueprint trace history */
 	const TSimpleRingBuffer<FBlueprintExecutionTrace>& GetTraceHistory() const { return ExecutionTraceHistory; }
 
@@ -79,6 +84,9 @@ public:
 	/** Remaps PIE actor instance paths to editor actor instances */
 	BLUEPRINTPROFILER_API FName RemapInstancePath(const FName InstanceName) const;
 
+	/** Returns the active instance name for statistic scope */
+	BLUEPRINTPROFILER_API FName GetActiveInstanceName() const;
+
 	/** Returns the function context containing the event */
 	TSharedPtr<class FBlueprintFunctionContext> GetFunctionContextForEventChecked(const FName ScopedEventName) const;
 
@@ -140,6 +148,8 @@ private:
 	TWeakObjectPtr<UBlueprint> Blueprint;
 	/** Blueprint Generated class for the context */
 	TWeakObjectPtr<UBlueprintGeneratedClass> BlueprintClass;
+	/** Blueprint utility contexts (Macro's ect ) */
+	TArray<FString> DependentUtilityContexts;
 	/** Root level execution node for this blueprint */
 	TSharedPtr<FScriptExecutionBlueprint> BlueprintNode;
 	/** UFunction contexts for this blueprint */
@@ -330,6 +340,9 @@ protected:
 	virtual TSharedPtr<FScriptExecutionNode> MapTunnelBoundary(const UEdGraphPin* TunnelPin);
 
 	/** Get tunnel boundary node */
+	virtual TSharedPtr<FScriptExecutionNode> GetTunnelBoundaryNode(const UEdGraphPin* TunnelPin);
+
+	/** Get tunnel boundary node, asserting on fail */
 	virtual TSharedPtr<FScriptExecutionNode> GetTunnelBoundaryNodeChecked(const UEdGraphPin* TunnelPin);
 
 	/** Maps the tunnel point into the instanced graph, creating the instanced graph if not already existing. */
@@ -396,7 +409,7 @@ private:
 	virtual TSharedPtr<FScriptExecutionNode> MapNodeExecution(UEdGraphNode* NodeToMap) override;
 	virtual void MapInputPins(TSharedPtr<FScriptExecutionNode> ExecNode, const TArray<UEdGraphPin*>& Pins) override;
 	virtual TSharedPtr<FScriptExecutionNode> MapPureNodeExecution(const UEdGraphPin* LinkedPin) override;
-	virtual TSharedPtr<FScriptExecutionNode> GetTunnelBoundaryNodeChecked(const UEdGraphPin* TunnelPin) override;
+	virtual TSharedPtr<FScriptExecutionNode> GetTunnelBoundaryNode(const UEdGraphPin* TunnelPin) override;
 	virtual TSharedPtr<FScriptExecutionNode> MapTunnelBoundary(const UEdGraphPin* TunnelPin) override;
 	// ~FBlueprintFunctionContext End
 	
@@ -504,6 +517,10 @@ private:
 	/** Add to trace history */
 	void AddToTraceHistory(const TSharedPtr<FScriptExecutionNode> ProfilerNode, const FScriptInstrumentedEvent& TraceSignal);
 
+	/** Add tunnel trace history */
+	void AddTunnelTraceHistory(const TSharedPtr<FScriptExecutionNode> TunnelBoundary, const FScriptInstrumentedEvent& TraceSignal, const FTracePath& InternalPath, const FTracePath& ExternalPath);
+
+
 private:
 
 	/** Instance name */
diff --git a/Engine/Source/Developer/HotReload/Private/HotReloadClassReinstancer.cpp b/Engine/Source/Developer/HotReload/Private/HotReloadClassReinstancer.cpp
--- a/Engine/Source/Developer/HotReload/Private/HotReloadClassReinstancer.cpp
+++ b/Engine/Source/Developer/HotReload/Private/HotReloadClassReinstancer.cpp
@@ -202,9 +202,9 @@ void FHotReloadClassReinstancer::ReconstructClassDefaultObject(UClass* InClass,
 	// Re-create
 	InClass->ClassDefaultObject = StaticAllocateObject(InClass, InOuter, InName, InFlags, EInternalObjectFlags::None, false);
 	check(InClass->ClassDefaultObject);
-	const bool bShouldInitilizeProperties = false;
+	const bool bShouldInitializeProperties = false;
 	const bool bCopyTransientsFromClassDefaults = false;
-	(*InClass->ClassConstructor)(FObjectInitializer(InClass->ClassDefaultObject, ParentDefaultObject, bCopyTransientsFromClassDefaults, bShouldInitilizeProperties));
+	(*InClass->ClassConstructor)(FObjectInitializer(InClass->ClassDefaultObject, ParentDefaultObject, bCopyTransientsFromClassDefaults, bShouldInitializeProperties));
 }
 
 void FHotReloadClassReinstancer::RecreateCDOAndSetupOldClassReinstancing(UClass* InOldClass)
diff --git a/Engine/Source/Editor/BlueprintGraph/Classes/EdGraphSchema_K2.h b/Engine/Source/Editor/BlueprintGraph/Classes/EdGraphSchema_K2.h
--- a/Engine/Source/Editor/BlueprintGraph/Classes/EdGraphSchema_K2.h
+++ b/Engine/Source/Editor/BlueprintGraph/Classes/EdGraphSchema_K2.h
@@ -175,6 +175,8 @@ public:
 	static const FName MD_ArrayParam;
 	static const FName MD_ArrayDependentParam;
 
+	static const FName MD_SetParam;
+
 	/** Metadata that identifies an integral property as a bitmask. */
 	static const FName MD_Bitmask;
 	/** Metadata that associates a bitmask property with a bitflag enum. */
@@ -425,6 +427,7 @@ public:
 	virtual bool ShouldHidePinDefaultValue(UEdGraphPin* Pin) const override;
 	virtual bool ShouldShowAssetPickerForPin(UEdGraphPin* Pin) const override;
 	virtual FLinearColor GetPinTypeColor(const FEdGraphPinType& PinType) const override;
+	FLinearColor GetSecondaryPinTypeColor(const FEdGraphPinType& PinType) const;
 	virtual FText GetPinDisplayName(const UEdGraphPin* Pin) const override;
 	virtual void ConstructBasicPinTooltip(const UEdGraphPin& Pin, const FText& PinDescription, FString& TooltipOut) const override;
 	virtual EGraphType GetGraphType(const UEdGraph* TestEdGraph) const override;
@@ -651,6 +654,13 @@ public:
 	/** Can Pin be recombined back to its original form */
 	bool CanRecombineStructPin(const UEdGraphPin& Pin) const;
 
+	/** 
+	 * Helper function for filling out Category, SubCategory, and SubCategoryObject based on a UProperty 
+	 * 
+	 * @return	true on success, false if the property is unsupported or invalid.
+	 */
+	static bool GetPropertyCategoryInfo(const UProperty* TestProperty, FString& OutCategory, FString& OutSubCategory, UObject*& OutSubCategoryObject, bool& bOutIsWeakPointer);
+
 	/**
 	 * Convert the type of a UProperty to the corresponding pin type.
 	 *
@@ -820,6 +830,19 @@ public:
 	 */
 	static FText TypeToText(UProperty* const Property);
 
+	/**
+	* Converts a terminal type into a fully qualified FText (e.g., object'ObjectName').
+	* Primarily used as a helper when converting containers to TypeToText.
+	*
+	* @param	Category					The category to convert into a FText.
+	* @param	SubCategory					The subcategory to convert into FText
+	* @param	SubCategoryObject			The SubcategoryObject to convert into FText
+	* @param	bIsWeakPtr					Whether the type is a WeakPtr
+	*
+	* @return	The converted type text.
+	*/
+	static FText TerminalTypeToText(const FString& Category, const FString& SubCategory, UObject* SubCategoryObject, bool bIsWeakPtr);
+
 	/**
 	 * Converts a pin type into a fully qualified FText (e.g., object'ObjectName').
 	 *
diff --git a/Engine/Source/Editor/BlueprintGraph/Classes/K2Node_CallFunction.h b/Engine/Source/Editor/BlueprintGraph/Classes/K2Node_CallFunction.h
--- a/Engine/Source/Editor/BlueprintGraph/Classes/K2Node_CallFunction.h
+++ b/Engine/Source/Editor/BlueprintGraph/Classes/K2Node_CallFunction.h
@@ -178,6 +178,9 @@ public:
 	/** Checks if the property is marked as "CustomStructureParam" */
 	static bool IsStructureWildcardProperty(const UFunction* InFunction, const FString& PropertyName);
 
+	/** returns true if InProperty should be treated as a wildcard (e.g. due to SetParam markup) */
+	static bool IsWildcardProperty(const UFunction* InFunction, const UProperty* InProperty);
+
 	/** Used to determine the result of AllowMultipleSelfs() (without having a node instance) */
 	static bool CanFunctionSupportMultipleTargets(UFunction const* InFunction);
 
@@ -203,6 +206,9 @@ private:
 	/** Invalidates current pin tool tips, so that they will be refreshed before being displayed: */
 	void InvalidatePinTooltips();
 
+	/** Conforms container pins */
+	void ConformContainerPins();
+
 protected:
 	/** Helper function to ensure function is called in our context */
 	virtual void FixupSelfMemberContext();
diff --git a/Engine/Source/Editor/BlueprintGraph/Classes/K2Node_ForEachElementInEnum.h b/Engine/Source/Editor/BlueprintGraph/Classes/K2Node_ForEachElementInEnum.h
--- a/Engine/Source/Editor/BlueprintGraph/Classes/K2Node_ForEachElementInEnum.h
+++ b/Engine/Source/Editor/BlueprintGraph/Classes/K2Node_ForEachElementInEnum.h
@@ -13,8 +13,8 @@ class UK2Node_ForEachElementInEnum : public UK2Node
 	UPROPERTY()
 	UEnum* Enum;
 
-	static const FString InsideLoopPinName;
-	static const FString EnumOuputPinName;
+	BLUEPRINTGRAPH_API static const FString InsideLoopPinName;
+	BLUEPRINTGRAPH_API static const FString EnumOuputPinName;
 
 	//~ Begin UEdGraphNode Interface
 	virtual void AllocateDefaultPins() override;
