Copying //UE4/Dev-Blueprints to //UE4/Dev-Main (Source: //UE4/Dev-Blueprints @ 3152873)

==========================
MAJOR FEATURES + CHANGES
==========================

Change 3131279 on 2016/09/19 by [REDACTED_PERSON]

	Fixing FText::Format warning for the Sub-Level Blueprints menu - Was using {LevelName} tag when there was no value for {LevelName} (defaulted to the first argument).

Change 3131318 on 2016/09/19 by [REDACTED_PERSON]

	[UE-35690] Minor revisions to Blueprint SCS execution to improve efficiency and address an out-of-order registration issue.

	change summary:
	- modified AActor::PostSpawnInitialize() to defer native component registration if there is no native scene root component set and if the actor is a BP type (i.e. will invoke SCS). this means that native actor classes with only non-scene components will now defer registration/post-registration until after SCS execution has established a valid scene root.
	- modified AActor::ExecuteConstruction() to gather the set of native scene components that SCS nodes can attach to before invoking the SCS. this was previously being done redundantly within the SCS itself at each level of the BP class inheritance hierarchy.
	- modified AActor::ExecuteConstruction() to do a final registration pass over all components after the all SCS levels have been executed. this was also previously being done within the SCS at each level. this avoids some extra redundancy.
	- modified USCS_Node::ExecuteNodeOnActor() to call RegisterAllComponents() on the given actor instance after establishing a valid scene root component if it was previously deferred at spawn time.
	- modified USCS_Node::ExecuteNodeOnActor() to now register components after they're created. since SCS execution goes from parent to child, parent scene components will always be registered before their children. non-native, non-scene component registration will also be deferred until a scene root component has been established.
	- added AActor::HasDeferredComponentRegistration()
	- modified AActor::RegisterAllComponents() to reset the actor's deferred component registration flag when called
	- modified AActor::AddComponent() to check the 'bAutoRegister' flag before calling RegisterComponent() (for consistency)
	- moved the RegisterInstancedComponent() utility method into USimpleConstructionScript and modified it to ensure that parent attachments are registered before their children.
	- modified USimpleConstructionScript::ExecuteScriptOnActor() to include an additional input parameter for passing in the set of native scene components that can be attached to.
	- modified USimpleConstructionScript::ExecuteScriptOnActor() to remove redundant/unnecessary work as noted above.

Change 3131842 on 2016/09/20 by [REDACTED_PERSON]

Change 3131847 on 2016/09/20 by [REDACTED_PERSON]

	SPropertyEditorAsset doesn't display UClass with "_C" prefix anymore.

Change 3131923 on 2016/09/20 by [REDACTED_PERSON]

Change 3132348 on 2016/09/20 by [REDACTED_PERSON]

	Fix CIS build issue (SA).

Change 3132383 on 2016/09/20 by [REDACTED_PERSON]

	Array UStruct::ScriptObjectReferences is filled while compilation. GC doesn't serialize script bytecode in editor anymore.

Change 3133072 on 2016/09/20 by [REDACTED_PERSON]

Change 3133216 on 2016/09/20 by [REDACTED_PERSON]

	+ BlueprintSetLibrary (add, remove, find, etc)
	+ HasGetTypeHash compile time function for detecting types that have GetTypeHish (modeled after HasOperatorEquals)
	= SPinTypeSelector can now disable container types based on current primary type, required transition to SComboButtton/SListView from SComboBox
	= Hide blueprint set library via BaseEditor.ini

Change 3133227 on 2016/09/20 by [REDACTED_PERSON]

	Test assets for TSet

Change 3133804 on 2016/09/21 by [REDACTED_PERSON]

	In UObjectLibrary, when bHasBlueprintClasses is true, BP references are automatically replaced by BPGC references.

Change 3133817 on 2016/09/21 by [REDACTED_PERSON]

	Fixed static Static Analysis warning

Change 3134377 on 2016/09/21 by [REDACTED_PERSON]

	ShowWorldContextObject is now inherited

Change 3134955 on 2016/09/21 by [REDACTED_PERSON]

	Making it so AdvancedDisplay metadata is taken into consideration and used in MakeStruct nodes.

Change 3134965 on 2016/09/21 by [REDACTED_PERSON]

	Fix for crash in FCDODiffControl when CDOs have different numbers of properties. First branch in the while loop would incorrectly advance Iter past the end of the array. Comments courtesy of [REDACTED_PERSON]

Change 3135523 on 2016/09/22 by [REDACTED_PERSON]

	PR #2755: Master (Contributed by [REDACTED_PERSON])
	Notable change: searching for Vector in BP editor context menu now gives different default result, prior result was mediocre, though (vector - vector)

Change 3136508 on 2016/09/22 by [REDACTED_PERSON]

	Removing a bIsVisible guard for level Blueprint menu actions - this was causing level BP options to disappear when you hid sub-levels. The guard doesn't seem to matter, as those actions will be removed with the world (when it is updated, or unloaded).

Change 3137587 on 2016/09/23 by [REDACTED_PERSON]

	Merged cl#3137578 from Odin branch

Change 3137666 on 2016/09/23 by [REDACTED_PERSON]

	This adds the ability to map composite graph instances in the same way we map macro instances for blueprint debug data and improves the quality of the debug data providing correct information for nested macro/composite instances at any script location in instrumented blueprint compilations.
	#Proj KismetCompiler, BlueprintGraph, UnrealEd

	- This is the first part of a two part change, the subsequent change will make use of the debug output to resolve complex trees of tunnel instances in the blueprint profiler.

Change 3137800 on 2016/09/23 by [REDACTED_PERSON]

	[UE-34896] Properties are now generated for client-only Blueprint components in an uncooked server-only context.

	change summary:
	- bumped BlueprintObjectsVersion
	- added a new 'ComponentClass' property to USCS_Node
	- added a new 'ComponentClass' field to the FComponentOverrideRecord struct (UInheritableComponentHandler)
	- added a USCS_Node::Serialize() override to fix up 'ComponentClass' on load (so that it's set prior to compile-on-load)
	- modified USimpleConstructionScript::CreateNodeImpl() to set the ComponentClass property in a new SCS node
	- modified USimpleConstructionScript::ValidateNodeTemplates() to consider the node to be valid if ComponentClass is set and is known to be filtered (i.e. the node will not be removed in this case)
	- modified USimpleConstructionScript::ValidateNodeTemplates() to emit a warning message in an uncooked client-only or server-only context if the ComponentClass could not be set in an existing package (i.e. if a resave is needed)
	- modified UInheritableComponentHandler::PostLoad() to fix up 'ComponentClass' on load
	- modified UInheritableComponentHandler::CreateOverridenComponentTemplate() to set the ComponentClass field in a new override record
	- modified UInheritableComponentHandler::IsRecordValid() to consider the record to be valid if ComponentClass is set (when ComponentTemplate is NULL)
	- modified UInheritableComponentHandler::IsRecordNecessary() to consider the record to be necessary if ComponentClass is set and is known to be filtered
	- modified FKismetCompilerContext::CreateClassVariablesFromBlueprint() to use 'ComponentClass' rather than 'ComponentTemplate' to infer the property subtype

Change 3137851 on 2016/09/23 by [REDACTED_PERSON]

	[UE-36079] Component overrides in a child blueprint will no longer trigger a warning message when the original component is removed from its parent.

	change summary:
	- modified UInheritableComponentHandler::IsRecordValid() to no longer consider a NULL OriginalTemplate to be invalid (so that the warning message is suppressed)
	- modified UInheritableComponentHandler::IsRecordNecessary() to consider a NULL OriginalTemplate to be unnecessary (so that the record is still removed in this case)

Change 3137948 on 2016/09/23 by [REDACTED_PERSON]

	CIS warning fix on mac for out of order initialisation.

Change 3139351 on 2016/09/25 by [REDACTED_PERSON]

	Updates the blueprint profiler to make use of the recent changes to macro/composite tunnel mapping and enhanced debug data.
	#Proj BlueprintProfiler, Engine

	- This is the second part of a two part change enabling multiple instances of nested macro/composite graphs in the blueprint profiler.

Change 3139376 on 2016/09/25 by [REDACTED_PERSON]

	CIS static analysis fix for CL 3137666

Change 3139377 on 2016/09/25 by [REDACTED_PERSON]

	Adding script code location checking for pure nodes that was missed in CL 3139351
	#Proj BlueprintProfiler

	- Fixes a missed issue with pure nodes inside macros/tunnels

Change 3139624 on 2016/09/26 by [REDACTED_PERSON]

	Fixed const local variables in Nativized code

	Merged cl#3139622 from Odin branch.

Change 3139641 on 2016/09/26 by [REDACTED_PERSON]

	Since we cannot distinguish which nodes are isolated by users (and shouldn't be validated) and which nodes are isolated during expansion step (and should be validated), the isolated nodes are pruned both before and after expantion step (and validation).

Change 3139961 on 2016/09/26 by [REDACTED_PERSON]

	CIS static analysis fix for CL 3137666

	- missed one of the warnings in a previous attempt.

Change 3140143 on 2016/09/26 by [REDACTED_PERSON]

	Fix for component property clearing on load, submitted on behalf of Mike.Beach

Change 3140694 on 2016/09/26 by [REDACTED_PERSON]

	Fix for GLEO when duplicating levels that have knots that reference delegates (specifically custom events)

Change 3140772 on 2016/09/26 by [REDACTED_PERSON]

	Further hardening SGraphPin::GraphPinObj access

Change 3140812 on 2016/09/26 by [REDACTED_PERSON]

	Corrected overzealous warning. Codepath is expected when functions are deleted but breakpoints aren't updated

Change 3140869 on 2016/09/26 by [REDACTED_PERSON]

	Update check to handle nested DSOs

Change 3141125 on 2016/09/27 by [REDACTED_PERSON]

	While reinstancing the CLASS_Abstract is cleared (just like the CLASS_Deprecated flag)

Change 3142715 on 2016/09/27 by [REDACTED_PERSON]

	Fix for crash when pasting nodes that have connections to nodes that aren't in the clipboard from one graph into another

Change 3143469 on 2016/09/28 by [REDACTED_PERSON]

	BP Profiler: Fixed Assert when profiling parent/child Blueprint

Change 3145215 on 2016/09/29 by [REDACTED_PERSON]

Change 3145580 on 2016/09/29 by [REDACTED_PERSON]

	Collapse secondary image instead of hiding it, allowing x button to be closer to primary image when secondary type image isn't present

Change 3146470 on 2016/09/30 by [REDACTED_PERSON]

	Restored cl#3085572 - it was lost while merging.

Change 3147046 on 2016/09/30 by [REDACTED_PERSON]

	BP doesn;t support weak obj ptr as parameters - Validation added. Function, created from colappsed nodes, cannot have a parameter of weakptr type.

Change 3148022 on 2016/10/01 by [REDACTED_PERSON]

	[UE-21109] Fix component instance data loss after renaming SCS component nodes at the Blueprint class level.

	change summary:
	- deprecated the public USCS_Node::VariableName member and replaced it with an internal-only member accessible via Get/Set method (i need to be in control of the set logic)
	- changed all occurrences of direct access to USCS_Node::VariableName to use a GetVariableName() call (since it's now internal)
	- simplified USCS_Node::GetVariableName() as what it used to do was legacy and thus is no longer necessary (it's been handled by USimpleConstructionScript::PostLoad for awhile now)
	- added USCS_Node::SetVariableName(); this now renames the component template (if valid) and all instances of it prior to changing the internal variable name. this ensures that archetype lookups will continue to function after a rename.
	- added USCS_Node::RenameComponentTemplate() to handle SCS node component template rename logic on a variable name change
	- switched the AActor::CheckComponentInstanceName() API to be publically-accessible; need to call this when renaming instanced components to match a new variable name in order to ensure that we rename any instance-only components out of the way first (this is the same logic that we run when constructing component instances on map load/RRCS, so it's consistent)
	- modified UInheritableComponentHandler::PostLoad() to fix up the component template within each record to match the original template object name. this ensures that ICH-specific archetype lookups will continue to function after a rename. it also ensures that any mismatched template names in existing assets will now be fixed up on load.
	- moved UActorComponent::ComponentTemplateSuffixName into the USimpleConstructionScript class (since the association is tied to templates created by that class specifically). note: this revises a change from a recent PR submission.
	- modified UpdateAttachedIsEditorOnly() to check the RF_ArchetypeObject flag rather than the ComponentTemplateSuffixName (part of the revision to the recent PR submission noted above)

Change 3148023 on 2016/10/01 by [REDACTED_PERSON]

	[UE-35562] Fix inherited Blueprinted component template defaults data loss in child Blueprint classes after recompiling the Blueprinted component class.

	change summary:
	- added a local FArchetypeReinstanceHelper struct + GetArchetypeObjects()/FindUniqueArchetypeObjectName() utility method implementations to KismetReinstanceUtilities.cpp
	- modified FBlueprintCompileReinstancer::ReplaceInstancesOfClass_Inner() to ensure that the full inherited component template (archetype) ancestry is renamed along with the base template when we rename the base template object away from its original name in order to make way for the new (reinstanced) template. the names must match the base template along the entire inheritance hierarchy in order for forward/reverse inherited component archetype lookups to succeed.

	notes:
	- BP (non-native) component templates (e.g. SCS/AddComponent nodes) that belong to the class being reinstanced (i.e. whose templates are not inherited from a parent BP) always inherit directly from the component CDO, and thus do not require this code path (that is, archetype lookups are not dependent on matching the CDO by name). similarly, native components (defined in C++) are included as part of the CDO defaults data, and thus also do not require this code path.

Change 3148030 on 2016/10/01 by [REDACTED_PERSON]

	UT fixes for CIS warnings related to SCS node API changes.

Change 3148256 on 2016/10/02 by [REDACTED_PERSON]

	This change adds the ability to filter debug/wire trace instrumenation and tracks expansion nodes for future use.
	#Proj KismetCompiler, BlueprintGraph, UnrealEd

Change 3148261 on 2016/10/02 by [REDACTED_PERSON]

	CIS fix, some code from another changelist leaked into CL 3148256

Change 3148480 on 2016/10/03 by [REDACTED_PERSON]

	This change attempts to address some profiler issues with class function context switching in the blueprint profiler.
	#Proj BlueprintProfiler, BlueprintGraph

Change 3148545 on 2016/10/03 by [REDACTED_PERSON]

	Skip unnecessary component validation work to fix invalid warnings when duplicating a BP class for reinstancing.

Change 3149001 on 2016/10/03 by [REDACTED_PERSON]

	This fixes an issue found with the instrumented blueprint compilations in which a certain compilation path would not provide the extended composite tunnel node heararchy in debug data and removes an unneceassary check that was causing problems.
	#Proj KismetCompiler, BlueprintProfiler

Change 3149031 on 2016/10/03 by [REDACTED_PERSON]

	Tunnel nodes without exec pin are not pruned before expansion.

Change 3149150 on 2016/10/03 by [REDACTED_PERSON]

	UGameplayTagsK2Node_LiteralGameplayTag is pure.

Change 3149290 on 2016/10/03 by [REDACTED_PERSON]

	GetBlueprintContext node is Pure.

Change 3149595 on 2016/10/03 by [REDACTED_PERSON]

	Fixing up some Orion content errors/warnings that should have been issues a while ago (error reporting was broken for a time, now fixed).

Change 3149667 on 2016/10/03 by [REDACTED_PERSON]

	Fixing up some Ocean content errors as fallout from a change in Dev-Blueprints - the errors properly identified a node that was using a culled input that was uninitialized.

Change 3149777 on 2016/10/03 by [REDACTED_PERSON]

	Fixing up  an Orion content warning - disconnecting a cast path in Hero_Automation, now that the node is producing a warning (the cast is impossible, and therefore the path is superfluous).

Change 3149988 on 2016/10/04 by [REDACTED_PERSON]

	Fixed IsNodePure functions.

Change 3150146 on 2016/10/04 by [REDACTED_PERSON]

	First pruning pass in done after ExpandTunnelsAndMacros is called.
	Isolated tunnels are pruned just like regular nodes.

Change 3150743 on 2016/10/04 by [REDACTED_PERSON]

	Mirroring CL 3150661 from Dev-VREditor

	Fix for crash on editor close after VR Foliage Editing.

Change 3151104 on 2016/10/04 by [REDACTED_PERSON]

	Added comment.

Change 3151979 on 2016/10/05 by [REDACTED_PERSON]

	Adding the keyword "custom" to K2Node_CustomEvent, so that it is prioritized when searching the Blueprint menu.

Change 3152286 on 2016/10/05 by [REDACTED_PERSON]

	Make sure, that an isolated node, that should be pure (but is not) won't be pruned.

diff --git a/Engine/Config/BaseEditor.ini b/Engine/Config/BaseEditor.ini
--- a/Engine/Config/BaseEditor.ini
+++ b/Engine/Config/BaseEditor.ini
@@ -170,6 +170,19 @@ DefaultProfile="StarterProfile"
 +StarterProfile=(NodeName="/Script/BlueprintGraph.K2Node_CallFunction", FieldName="/Script/Engine.GameplayStatics:GetAllActorsOfClass")
 +StarterProfile=(NodeName="/Script/BlueprintGraph.K2Node_SpawnActorFromClass")
 
+[BlueprintEditor.Menu]
++BlueprintHiddenFields="/Script/Engine.BlueprintSetLibrary:Set_Add"
++BlueprintHiddenFields="/Script/Engine.BlueprintSetLibrary:Set_AddItems"
++BlueprintHiddenFields="/Script/Engine.BlueprintSetLibrary:Set_Remove"
++BlueprintHiddenFields="/Script/Engine.BlueprintSetLibrary:Set_RemoveItems"
++BlueprintHiddenFields="/Script/Engine.BlueprintSetLibrary:Set_ToArray"
++BlueprintHiddenFields="/Script/Engine.BlueprintSetLibrary:Set_Clear"
++BlueprintHiddenFields="/Script/Engine.BlueprintSetLibrary:Set_Length"
++BlueprintHiddenFields="/Script/Engine.BlueprintSetLibrary:Set_Contains"
++BlueprintHiddenFields="/Script/Engine.BlueprintSetLibrary:Set_Intersection"
++BlueprintHiddenFields="/Script/Engine.BlueprintSetLibrary:Set_Union"
++BlueprintHiddenFields="/Script/Engine.BlueprintSetLibrary:Set_Difference"
+
 [UnrealEd.MeshPaint]
 MinBrushRadius=1
 MaxBrushRadius=256
diff --git a/Engine/Config/BaseEngine.ini b/Engine/Config/BaseEngine.ini
--- a/Engine/Config/BaseEngine.ini
+++ b/Engine/Config/BaseEngine.ini
@@ -795,6 +795,8 @@ bUseStreamingPause=false
 +TaggedPropertyRedirects=(ClassName="ComponentKey",OldPropertyName="VariableName",NewPropertyName="SCSVariableName")
 +TaggedPropertyRedirects=(ClassName="ComponentKey",OldPropertyName="VariableGuid",NewPropertyName="AssociatedGuid")
 
++TaggedPropertyRedirects=(ClassName="SCS_Node",OldPropertyName="VariableName",NewPropertyName="InternalVariableName")
+
 ;Vim rename
 +TaggedPropertyRedirects=(ClassName="AnimInstance",OldPropertyName="VimVertexAnims",NewPropertyName="VertexAnims")
 
@@ -1253,6 +1255,7 @@ bUseStreamingPause=false
 +K2FieldRedirects=(OldFieldName="KismetMathLibrary.FRandRangeFromStream",NewFieldName="RandomFloatInRangeFromStream")
 
 +K2FieldRedirects=(OldFieldName="KismetMathLibrary.NearlyEqual_VectorVector",NewFieldName="EqualEqual_VectorVector")
++K2FieldRedirects=(OldFieldName="KismetMathLibrary.GetDirectionVector",NewFieldName="GetDirectionUnitVector")
 +K2FieldRedirects=(OldFieldName="KismetMathLibrary.NearlyEqual_RotatorRotator",NewFieldName="EqualEqual_RotatorRotator")
 
 +K2FieldRedirects=(OldFieldName="KismetMathLibrary.ProjectOnTo",NewFieldName="ProjectVectorOnToVector")
diff --git a/Engine/Source/Developer/BlueprintCompilerCppBackend/Private/BlueprintCompilerCppBackendBase.cpp b/Engine/Source/Developer/BlueprintCompilerCppBackend/Private/BlueprintCompilerCppBackendBase.cpp
--- a/Engine/Source/Developer/BlueprintCompilerCppBackend/Private/BlueprintCompilerCppBackendBase.cpp
+++ b/Engine/Source/Developer/BlueprintCompilerCppBackend/Private/BlueprintCompilerCppBackendBase.cpp
@@ -405,7 +405,7 @@ static void DeclareLocalVariables(FEmitterLocalContext& EmitterContext, TArray<U
 		UProperty* LocalVariable = LocalVariables[i];
 		if (!bUseExecutionGroup || PropertiesUsedByCurrentExecutionGroup.Contains(LocalVariable))
 		{
-			const FString CppDeclaration = EmitterContext.ExportCppDeclaration(LocalVariable, EExportedDeclaration::Local, EPropertyExportCPPFlags::CPPF_CustomTypeName | EPropertyExportCPPFlags::CPPF_BlueprintCppBackend);
+			const FString CppDeclaration = EmitterContext.ExportCppDeclaration(LocalVariable, EExportedDeclaration::Local, EPropertyExportCPPFlags::CPPF_CustomTypeName | EPropertyExportCPPFlags::CPPF_BlueprintCppBackend | EPropertyExportCPPFlags::CPPF_NoConst);
 			UStructProperty* StructProperty = Cast<UStructProperty>(LocalVariable);
 			const TCHAR* EmptyDefaultConstructor = FEmitHelper::EmptyDefaultConstructor(StructProperty ? StructProperty->Struct : nullptr);
 			EmitterContext.AddLine(CppDeclaration + EmptyDefaultConstructor + TEXT(";"));
diff --git a/Engine/Source/Developer/BlueprintCompilerCppBackend/Private/BlueprintCompilerCppBackendValueHelper.cpp b/Engine/Source/Developer/BlueprintCompilerCppBackend/Private/BlueprintCompilerCppBackendValueHelper.cpp
--- a/Engine/Source/Developer/BlueprintCompilerCppBackend/Private/BlueprintCompilerCppBackendValueHelper.cpp
+++ b/Engine/Source/Developer/BlueprintCompilerCppBackend/Private/BlueprintCompilerCppBackendValueHelper.cpp
@@ -649,7 +649,7 @@ FString FEmitDefaultValueHelper::HandleNonNativeComponent(FEmitterLocalContext&
 	UBlueprintGeneratedClass* BPGC = CastChecked<UBlueprintGeneratedClass>(Context.GetCurrentlyGeneratedClass());
 	if (UActorComponent* ComponentTemplate = Node->GetActualComponentTemplate(BPGC))
 	{
-		const FString VariableCleanName = Node->VariableName.ToString();
+		const FString VariableCleanName = Node->GetVariableName().ToString();
 
 		const UObjectProperty* VariableProperty = FindField<UObjectProperty>(BPGC, *VariableCleanName);
 		if (VariableProperty)
@@ -695,7 +695,7 @@ FString FEmitDefaultValueHelper::HandleNonNativeComponent(FEmitterLocalContext&
 				FString ParentVariableName;
 				if (ParentNode)
 				{
-					const FString CleanParentVariableName = ParentNode->VariableName.ToString();
+					const FString CleanParentVariableName = ParentNode->GetVariableName().ToString();
 					const UObjectProperty* ParentVariableProperty = FindField<UObjectProperty>(BPGC, *CleanParentVariableName);
 					ParentVariableName = ParentVariableProperty ? FEmitHelper::GetCppName(ParentVariableProperty) : CleanParentVariableName;
 				}
diff --git a/Engine/Source/Developer/BlueprintNativeCodeGen/Private/BlueprintNativeCodeGenModule.cpp b/Engine/Source/Developer/BlueprintNativeCodeGen/Private/BlueprintNativeCodeGenModule.cpp
--- a/Engine/Source/Developer/BlueprintNativeCodeGen/Private/BlueprintNativeCodeGenModule.cpp
+++ b/Engine/Source/Developer/BlueprintNativeCodeGen/Private/BlueprintNativeCodeGenModule.cpp
@@ -614,7 +614,7 @@ UObject* FBlueprintNativeCodeGenModule::FindReplacedNameAndOuter(UObject* Object
 				{
 					if (Node->ComponentTemplate == ActorComponent)
 					{
-						OutName = Node->VariableName;
+						OutName = Node->GetVariableName();
 						if (OutName != NAME_None)
 						{
 							Outer = BPGC->GetDefaultObject(false);
diff --git a/Engine/Source/Developer/BlueprintProfiler/Private/BlueprintProfiler.cpp b/Engine/Source/Developer/BlueprintProfiler/Private/BlueprintProfiler.cpp
--- a/Engine/Source/Developer/BlueprintProfiler/Private/BlueprintProfiler.cpp
+++ b/Engine/Source/Developer/BlueprintProfiler/Private/BlueprintProfiler.cpp
@@ -73,7 +73,7 @@ void FBlueprintProfiler::InstrumentEvent(const FScriptInstrumentationSignal& Eve
 		// Reset context on event end
 		if (Event.GetType() == EScriptInstrumentation::Stop)
 		{
-			CaptureContext.ResetContext();
+			CaptureContext.PopScope();
 		}
 	}
 }
@@ -244,12 +244,12 @@ void FBlueprintProfiler::ProcessEventProfilingData()
 				if (ScriptEventRanges.Num() > 0)
 				{
 					// Check against the last class scope first before adding another range.
-					FEventRange& LastEventRange = ScriptEventRanges.Last();
-					const FName NewClassScope = InstrumentationEventQueue[EventIdx].GetFunctionClassScopeName();
+					const FEventRange& LastEventRange = ScriptEventRanges.Last();
+					const FName NewClassScope = CurrEvent.GetFunctionClassScopeName();
 					if (LastEventRange.ClassScope != NewClassScope)
 					{
 						// Override the class scope signal to be an event.
-						InstrumentationEventQueue[EventIdx].OverrideType(EScriptInstrumentation::Event);
+						CurrEvent.OverrideType(EScriptInstrumentation::Event);
 						// Create new sub event event range to be processed as a sub event.
 						FEventRange NewClassScopeRange;
 						NewClassScopeRange.BlueprintContext = LastEventRange.BlueprintContext;
@@ -274,12 +274,11 @@ void FBlueprintProfiler::ProcessEventProfilingData()
 							StopTime = FMath::Max<double>(StopTime, EventTime);
 						}
 						// Set event timings
-						InstrumentationEventQueue[EventIdx].OverrideTime(StartTime);
+						CurrEvent.OverrideTime(StartTime);
 						// Insert Stop marker
 						FScriptInstrumentedEvent StopMarker(EScriptInstrumentation::Stop, NewClassScope, InstrumentationEventQueue[EventIdx].GetFunctionName(), INDEX_NONE);
 						StopMarker.OverrideTime(StopTime);
 						InstrumentationEventQueue.Insert(StopMarker, SubClassIdx);
-						// Note: we could skip the class scope signals and jump straight to processing here - it might cause problems though.
 					}
 					else
 					{
@@ -355,6 +354,8 @@ void FBlueprintProfiler::ProcessEventProfilingData()
 								SuspendedEventMap.Add(EventToProcess->GetLatentLinkId()) = EventToProcess;
 							}
 						}
+						// Note: We should probably assert here if the above Process call returned false. We should always
+						// Process successfully after a Stop event. If false is returned then something is wrong.
 					}
 					ScriptEventRanges.Pop();
 				}
@@ -362,6 +363,10 @@ void FBlueprintProfiler::ProcessEventProfilingData()
 			}
 		}
 	}
+	// For the time being, we should have processed all events in the queue.
+	// In the future we will potentially only process a few events at a time, and this check should be removed.
+	// Removed this check after finding an issue with it switching out maps.
+//	check(InstrumentationEventQueue.Num() == 0);
 	// Update all active contexts if the display settings changed.
 	UBlueprintProfilerSettings* Settings = GetMutableDefault<UBlueprintProfilerSettings>();
 	if (Settings->GetPerformanceThresholdsModified())
diff --git a/Engine/Source/Developer/BlueprintProfiler/Private/ScriptInstrumentationCapture.cpp b/Engine/Source/Developer/BlueprintProfiler/Private/ScriptInstrumentationCapture.cpp
--- a/Engine/Source/Developer/BlueprintProfiler/Private/ScriptInstrumentationCapture.cpp
+++ b/Engine/Source/Developer/BlueprintProfiler/Private/ScriptInstrumentationCapture.cpp
@@ -23,6 +23,13 @@ FName FScriptInstrumentedEvent::GetScopedFunctionName() const
 
 bool FInstrumentationCaptureContext::UpdateContext(const FScriptInstrumentationSignal& InstrumentSignal, TArray<FScriptInstrumentedEvent>& InstrumentationQueue)
 {
+	enum EScopeChange
+	{
+		None,
+		Push,
+		Pop
+	};
+
 	if (InstrumentSignal.GetType() == EScriptInstrumentation::InlineEvent)
 	{
 		// This inline event needs to be processed using the next incoming event's function class scope,
@@ -30,56 +37,64 @@ bool FInstrumentationCaptureContext::UpdateContext(const FScriptInstrumentationS
 		PendingInlineEventName = InstrumentSignal.GetFunctionName();
 		return false;
 	}
-	else if (PendingInlineEventName != NAME_None)
+	else if (InstrumentSignal.GetType() != EScriptInstrumentation::Stop)
 	{
-		// The pending inline event needs to be added first to make sure our events are processed in the proper order.
-		InstrumentationQueue.Add(FScriptInstrumentedEvent(EScriptInstrumentation::Event, InstrumentSignal.GetFunctionClassScope()->GetFName(), PendingInlineEventName));
-		PendingInlineEventName = NAME_None;
-	}
+		FContextScope SignalScope;
+		SignalScope.ContextObject = InstrumentSignal.GetContextObject();
+		SignalScope.ContextClass = InstrumentSignal.GetClass();
+		SignalScope.ContextFunctionClassScope = InstrumentSignal.GetFunctionClassScope();
 
-	// Handle instance context switch
-	if (const UObject* NewContextObject = InstrumentSignal.GetContextObject())
-	{
-		if (NewContextObject != ContextObject)
+		const FContextScope& CurrentScope = ScopeStack.Num() > 0 ? ScopeStack.Last() : SignalScope;
+	
+		// Always push a new scope for an event
+		EScopeChange ScopeChange = PendingInlineEventName != NAME_None ? EScopeChange::Push : EScopeChange::None;
+		bool bIsStackEmpty = ScopeStack.Num() == 0;
+		if (bIsStackEmpty || SignalScope.ContextObject != CurrentScope.ContextObject)
 		{
-			// Handle class context switch
-			if (const UClass* NewClass = InstrumentSignal.GetClass())
+			if (bIsStackEmpty || SignalScope.ContextClass != CurrentScope.ContextClass)
+			{
+				InstrumentationQueue.Add(FScriptInstrumentedEvent(EScriptInstrumentation::Class, SignalScope.ContextClass->GetFName(), InstrumentSignal.GetFunctionName(), SignalScope.ContextClass->GetPathName()));
+			}
+
+			InstrumentationQueue.Add(FScriptInstrumentedEvent(EScriptInstrumentation::Instance, SignalScope.ContextClass->GetFName(), InstrumentSignal.GetFunctionName(), SignalScope.ContextObject->GetPathName()));
+			if (InstrumentSignal.GetType() == EScriptInstrumentation::PopState)
+			{
+				// An Instance change for a PopState signal indicates that we are returning to a previous scope
+				ScopeChange = EScopeChange::Pop;
+			}
+			else
 			{
-				if (NewClass != ContextClass)
-				{
-					InstrumentationQueue.Add(FScriptInstrumentedEvent(EScriptInstrumentation::Class, InstrumentSignal.GetFunctionClassScope()->GetFName(), NewClass->GetPathName()));
-					ContextClass = NewClass;
-				}
+				ScopeChange = EScopeChange::Push;
 			}
-			// Handle instance change
-			InstrumentationQueue.Add(FScriptInstrumentedEvent(EScriptInstrumentation::Instance, ContextClass->GetFName(), NewContextObject->GetPathName()));
-			ContextObject = NewContextObject;
 		}
-		else
+		else if (SignalScope.ContextFunctionClassScope != CurrentScope.ContextFunctionClassScope)
 		{
-			TWeakObjectPtr<const UClass> SignalScope = InstrumentSignal.GetFunctionClassScope();
-			// Handle function class scope changes
-			if (ContextFunctionClassScope != SignalScope)
+			if (InstrumentSignal.GetType() != EScriptInstrumentation::Event && PendingInlineEventName == NAME_None)
 			{
-				FName ClassScopeName = SignalScope->GetFName();
-				if (InstrumentSignal.GetType() == EScriptInstrumentation::Stop)
-				{
-					// In a Stop event, we need to use the last class scope on the stack to signify
-					// that we are reverting back to the previous class scope
-					ClassScopeName = ScopeStack.Pop();
-				}
-				else
-				{
-					// Stop events cause the Context to be reset so we need to check against the top scope on the stack before pushing it to the stack
-					if (ScopeStack.Num() == 0 || ScopeStack.Last() != ClassScopeName)
-					{
-						ScopeStack.Push(ClassScopeName);
-					}
-					ContextFunctionClassScope = SignalScope;
-				}
+				InstrumentationQueue.Add(FScriptInstrumentedEvent(EScriptInstrumentation::ClassScope, SignalScope.ContextFunctionClassScope->GetFName(), InstrumentSignal.GetFunctionName()));
+			}
+			ScopeChange = EScopeChange::Push;
+		}
 
-				InstrumentationQueue.Add(FScriptInstrumentedEvent(EScriptInstrumentation::ClassScope, ClassScopeName, InstrumentSignal.GetFunctionName()));
+		switch (ScopeChange)
+		{
+			case EScopeChange::Push:
+			{
+				ScopeStack.Push(SignalScope);
+			}
+			break;
+			case EScopeChange::Pop:
+			{
+				PopScope();
 			}
+			break;
+		}
+
+		if (PendingInlineEventName != NAME_None)
+		{
+			// The pending inline event needs to be added first to make sure our events are processed in the proper order.
+			InstrumentationQueue.Add(FScriptInstrumentedEvent(EScriptInstrumentation::Event, InstrumentSignal.GetFunctionClassScope()->GetFName(), PendingInlineEventName));
+			PendingInlineEventName = NAME_None;
 		}
 	}
 	
@@ -88,7 +103,13 @@ bool FInstrumentationCaptureContext::UpdateContext(const FScriptInstrumentationS
 
 void FInstrumentationCaptureContext::ResetContext()
 {
-	ContextClass.Reset();
-	ContextFunctionClassScope.Reset();
-	ContextObject.Reset();
+	ScopeStack.Empty();
+}
+
+void FInstrumentationCaptureContext::PopScope()
+{
+	if (ScopeStack.Num() > 0)
+	{
+		ScopeStack.Pop();
+	}
 }
\ No newline at end of file
diff --git a/Engine/Source/Developer/BlueprintProfiler/Private/ScriptInstrumentationPlayback.cpp b/Engine/Source/Developer/BlueprintProfiler/Private/ScriptInstrumentationPlayback.cpp
--- a/Engine/Source/Developer/BlueprintProfiler/Private/ScriptInstrumentationPlayback.cpp
+++ b/Engine/Source/Developer/BlueprintProfiler/Private/ScriptInstrumentationPlayback.cpp
@@ -429,6 +429,7 @@ bool FBlueprintExecutionContext::MapBlueprintExecution()
 			NewFunctionContext->DiscoverTunnels(GraphEntry.Value, DiscoveredTunnels);
 		}
 		// Create and map each tunnels instance as its own function context.
+		TArray<UK2Node_Tunnel*> Tunnels;
 		for (auto TunnelInstance : DiscoveredTunnels)
 		{
 			if (UEdGraph* TunnelGraph = FBlueprintFunctionContext::GetGraphFromNode(TunnelInstance.Key, false))
@@ -441,9 +442,11 @@ bool FBlueprintExecutionContext::MapBlueprintExecution()
 					DependentUtilityContexts.AddUnique(ClassPath);
 				}
 				// Create the tunnel instance context.
-				const FName TunnelInstanceFunctionName = FBlueprintFunctionContext::GetTunnelInstanceFunctionName(TunnelInstance.Key);
+				Tunnels.Push(TunnelInstance.Key);
+				const FName TunnelInstanceFunctionName = TunnelInstance.Value->GetTunnelInstanceFunctionName(TunnelInstance.Key);
 				TSharedPtr<FBlueprintTunnelInstanceContext> NewTunnelInstanceContext = CreateFunctionContext<FBlueprintTunnelInstanceContext>(TunnelInstanceFunctionName, TunnelGraph);
-				NewTunnelInstanceContext->MapTunnelContext(TunnelInstance.Value, TunnelInstance.Key);
+				NewTunnelInstanceContext->MapTunnelContext(TunnelInstance.Value, TunnelInstance.Value, Tunnels);
+				Tunnels.Pop();
 				TunnelInstance.Value->MapTunnelInstance(TunnelInstance.Key);
 			}
 		}
@@ -554,14 +557,14 @@ void FBlueprintFunctionContext::InitialiseContextFromGraph(TSharedPtr<FBlueprint
 							EventParams.Tooltip = LOCTEXT("NavigateToInheritedEventLocationHyperlink_ToolTip", "Navigate to the Inherited Event");
 							EventParams.IconColor = GetDefault<UGraphEditorSettings>()->ParentFunctionCallNodeTitleColor;
 							EventParams.NodeFlags = EScriptExecutionNodeFlags::Event|EScriptExecutionNodeFlags::RuntimeEvent|EScriptExecutionNodeFlags::InheritedEvent;
-							EventParams.DisplayName = FText::FromString(FString::Printf(TEXT("%s (%s)"), *EventNode->GetNodeTitle(ENodeTitleType::ListView).ToString(), *OwningBlueprint->GetName()));
+							EventParams.DisplayName = FText::FromString(FString::Printf(TEXT("%s (%s)"), *EventNode->GetNodeTitle(ENodeTitleType::MenuTitle).ToString(), *OwningBlueprint->GetName()));
 						}
 						else
 						{
 							EventParams.Tooltip = LOCTEXT("NavigateToEventLocationHyperlink_ToolTip", "Navigate to the Event");
 							EventParams.IconColor = GetDefault<UGraphEditorSettings>()->EventNodeTitleColor;
 							EventParams.NodeFlags = EScriptExecutionNodeFlags::Event|EScriptExecutionNodeFlags::RuntimeEvent;
-							EventParams.DisplayName = EventNode->GetNodeTitle(ENodeTitleType::ListView);
+							EventParams.DisplayName = EventNode->GetNodeTitle(ENodeTitleType::MenuTitle);
 						}
 						GetNodeCustomizations(EventParams);
 						const FSlateBrush* EventIcon = EventNode->ShowPaletteIconOnNode() ?	EventNode->GetIconAndTint(EventParams.IconColor).GetOptionalIcon() :
@@ -812,7 +815,7 @@ void FBlueprintFunctionContext::CreateDelegatePinEvents(TSharedPtr<FBlueprintExe
 				EventParams.NodeName = FName(*FString::Printf(TEXT("%s::DummyEvent"), *NodeEvents.Key->GetName()));
 				EventParams.ObservedObject = NodeEvents.Key;
 				EventParams.OwningGraphName = NAME_None;
-				EventParams.DisplayName = NodeEvents.Key->GetNodeTitle(ENodeTitleType::ListView);
+				EventParams.DisplayName = NodeEvents.Key->GetNodeTitle(ENodeTitleType::MenuTitle);
 				EventParams.Tooltip = LOCTEXT("NavigateToEventLocationHyperlink_ToolTip", "Navigate to the Event");
 				EventParams.NodeFlags = FunctionContext->IsInheritedContext() ? (EScriptExecutionNodeFlags::Event|EScriptExecutionNodeFlags::RuntimeEvent|EScriptExecutionNodeFlags::InheritedEvent) : 
 																				(EScriptExecutionNodeFlags::Event|EScriptExecutionNodeFlags::RuntimeEvent);
@@ -999,7 +1002,7 @@ void FBlueprintFunctionContext::DetermineGraphNodeCharacteristics(const UEdGraph
 		NodeParams.NodeFlags |= EScriptExecutionNodeFlags::CustomEvent;
 	}
 	// Create display name.
-	NodeParams.DisplayName = GraphNode->GetNodeTitle(ENodeTitleType::ListView);
+	NodeParams.DisplayName = GraphNode->GetNodeTitle(ENodeTitleType::MenuTitle);
 }
 
 TSharedPtr<FScriptExecutionNode> FBlueprintFunctionContext::MapNodeExecution(UEdGraphNode* NodeToMap)
@@ -1135,7 +1138,7 @@ TSharedPtr<FScriptExecutionNode> FBlueprintFunctionContext::MapPureNodeExecution
 				PinParams.ObservedPin = LinkedPin;
 				PinParams.Tooltip = LOCTEXT("NavigateToPinLocationHyperlink_ToolTip", "Navigate to the Pure Pin");
 				FFormatNamedArguments Args;
-				Args.Add("NodeDisplayName", LinkedNode->GetNodeTitle(ENodeTitleType::ListView));
+				Args.Add("NodeDisplayName", LinkedNode->GetNodeTitle(ENodeTitleType::MenuTitle));
 				Args.Add("PinDisplayName", LinkedPin->PinFriendlyName);
 				FText Format = LinkedPin->PinFriendlyName.IsEmpty() ? LOCTEXT("PureNodeDisplay_Text", "{NodeDisplayName}") : LOCTEXT("PurePinDisplay_Text", "{NodeDisplayName} - {PinDisplayName}");
 				PinParams.DisplayName = FText::Format(Format, Args);
@@ -1504,7 +1507,7 @@ FName FBlueprintFunctionContext::GetGraphNameFromNode(const UEdGraphNode* GraphN
 	return NodeGraphName;
 }
 
-FName FBlueprintFunctionContext::GetTunnelInstanceFunctionName(const UEdGraphNode* GraphNode)
+FName FBlueprintFunctionContext::GetTunnelInstanceFunctionName(const UEdGraphNode* GraphNode) const
 {
 	FName ScopedFunctionName = NAME_None;
 	UEdGraph* OuterGraph = GraphNode ? GraphNode->GetTypedOuter<UEdGraph>() : nullptr;
@@ -1517,7 +1520,7 @@ FName FBlueprintFunctionContext::GetTunnelInstanceFunctionName(const UEdGraphNod
 			if (BPGC)
 			{
 				// To identify instances exactly, we need the owning blueprint name, owning graph name, and the node name to avoid collisions.
-				ScopedFunctionName = FName(*FString::Printf(TEXT("%s::%s::%s"), *BPGC->GetName(), *OuterGraph->GetName(), *GraphNode->GetName()));
+				ScopedFunctionName = FName(*FString::Printf(TEXT("%s::%s"), *FunctionName.ToString(), *GraphNode->GetName()));
 			}
 		}
 	}
@@ -1670,27 +1673,28 @@ template<typename ExecNodeType> TSharedPtr<ExecNodeType> FBlueprintFunctionConte
 	return Result;
 }
 
-bool FBlueprintFunctionContext::GetProfilerContextFromScriptOffset(const int32 ScriptOffset, TSharedPtr<FScriptExecutionNode>& ExecNode, TSharedPtr<FBlueprintFunctionContext>& FunctionContext)
+bool FBlueprintFunctionContext::GetProfilerContextFromScriptOffset(const int32 ScriptOffset, FNodeExecutionContext& ExecContextOut)
 {
+	// Locate the correct context
 	if (const UEdGraphNode* GraphNode = GetNodeFromCodeLocation(ScriptOffset))
 	{
 		// get the fully scoped function name
 		const FName FunctionContextName = GetScopedFunctionNameFromNode(GraphNode);
 		if (FunctionContextName == FunctionName)
 		{
-			FunctionContext = AsShared();
-			ExecNode = GetProfilerDataForNode(GraphNode->GetFName());
+			ExecContextOut.ProfilerNode = GetProfilerDataForNode(GraphNode->GetFName());
+			ExecContextOut.ProfilerContext = AsShared();
 		}
 		else
 		{
 			TSharedPtr<FBlueprintFunctionContext> OwningFunction = BlueprintContext.Pin()->GetFunctionContext(FunctionContextName);
 			if (OwningFunction.IsValid())
 			{
-				OwningFunction->GetProfilerContextFromScriptOffset(ScriptOffset, ExecNode, FunctionContext);
+				OwningFunction->GetProfilerContextFromScriptOffset(ScriptOffset, ExecContextOut);
 			}
 		}
 	}
-	return ExecNode.IsValid() && FunctionContext.IsValid();
+	return ExecContextOut.IsValid();
 }
 
 const UEdGraphNode* FBlueprintFunctionContext::GetNodeFromCodeLocation(const int32 ScriptOffset)
@@ -1762,7 +1766,7 @@ const UEdGraphPin* FBlueprintFunctionContext::GetPinFromCodeLocation(const int32
 	return Result.Get();
 }
 
-const int32 FBlueprintFunctionContext::GetCodeLocationFromPin(const UEdGraphPin* Pin) const
+const int32 FBlueprintFunctionContext::GetCodeLocationFromPin(const UEdGraphPin* Pin)
 {
 	if (BlueprintClass.IsValid())
 	{
@@ -1780,6 +1784,14 @@ void FBlueprintFunctionContext::GetAllCodeLocationsFromPin(const UEdGraphPin* Pi
 	}
 }
 
+void FBlueprintFunctionContext::FindAllTunnelsAtCodeLocation(const int32 CodeLocation, TArray<UEdGraphNode*>& MacrosAtLocation) const
+{
+	if (BlueprintClass.IsValid() && Function.IsValid())
+	{
+		BlueprintClass.Get()->GetDebugData().FindMacroInstanceNodesFromCodeLocation(Function.Get(), CodeLocation, MacrosAtLocation);
+	}
+}
+
 FInt32Range FBlueprintFunctionContext::GetPureNodeScriptCodeRange(const UEdGraphNode* Node) const
 {
 	if (BlueprintClass.IsValid())
@@ -1854,6 +1866,41 @@ void FBlueprintFunctionContext::AddCallSiteEntryPointsToNode(TSharedPtr<FScriptE
 	}
 }
 
+FNodeExecutionContext& FBlueprintFunctionContext::GetNodeExecutionContext(const int32 ScriptOffset)
+{
+	FNodeExecutionContext& Context = ScriptOffsetToNodeContext.FindOrAdd(ScriptOffset);
+	if (!Context.IsValid())
+	{
+		const UEdGraphNode* GraphNode = GetNodeFromCodeLocation(ScriptOffset);
+		TSharedPtr<FBlueprintFunctionContext> SearchContext = AsShared();
+		TArray<FName> ScopedContextNames;
+		GraphNodeToContextName.MultiFind(GraphNode, ScopedContextNames);
+		if (ScopedContextNames.Num() > 0)
+		{
+			for (const FName ScopedContextName : ScopedContextNames)
+			{
+				SearchContext = BlueprintContext.Pin()->GetFunctionContext(ScopedContextName);
+				if (SearchContext.IsValid())
+				{
+					if (SearchContext->GetProfilerContextFromScriptOffset(ScriptOffset, Context))
+					{
+						Context.GraphNode = GraphNode;
+						break;
+					}
+				}
+			}
+		}
+		else
+		{
+			if (SearchContext->GetProfilerContextFromScriptOffset(ScriptOffset, Context))
+			{
+				Context.GraphNode = GraphNode;
+			}
+		}
+	}
+	return Context;
+}
+
 //////////////////////////////////////////////////////////////////////
 // FBlueprintTunnelInstanceContext
 
@@ -1869,17 +1916,18 @@ void FBlueprintTunnelInstanceContext::SetParentContext(TSharedPtr<FBlueprintFunc
 	ParentTunnel = StaticCastSharedPtr<FBlueprintTunnelInstanceContext>(ParentContextIn);
 }
 
-void FBlueprintTunnelInstanceContext::MapTunnelContext(TSharedPtr<FBlueprintFunctionContext> CallingFunctionContext, UK2Node_Tunnel* TunnelInstance)
+void FBlueprintTunnelInstanceContext::MapTunnelContext(TSharedPtr<FBlueprintFunctionContext> RootFunctionContextIn, TSharedPtr<FBlueprintFunctionContext> CallingFunctionContext, TArray<UK2Node_Tunnel*> TunnelInstances)
 {
 	if (!TunnelInstanceNode.IsValid())
 	{
 		// Set the mapping context.
+		RootFunctionContext = RootFunctionContextIn;
 		Function = CallingFunctionContext->GetUFunction();
 		BlueprintClass = CallingFunctionContext->GetBlueprintClass();
+		check (TunnelInstances.Num() > 0);
+		NestedTunnelInstanceStack = TunnelInstances;
+		UK2Node_Tunnel* TunnelInstance = TunnelInstances.Last();
 		TunnelInstanceNode = TunnelInstance;
-		// Register macro graph context
-		const FName ScopedFunctionName = GetScopedFunctionNameFromNode(TunnelInstance);
-		BlueprintContext.Pin()->AddFunctionContext(ScopedFunctionName, AsShared());
 		// Find the function context that represents the graph and not the instance
 		UEdGraph* TunnelGraph = GetGraphFromNode(TunnelInstance, false);
 		// Map tunnel Input/Output, creating stubbed pure pin chain sites we can link up in nested tunnels.
@@ -1890,7 +1938,7 @@ void FBlueprintTunnelInstanceContext::MapTunnelContext(TSharedPtr<FBlueprintFunc
 		TunnelInstanceParams.SampleFrequency = 1;
 		TunnelInstanceParams.NodeName = TunnelInstance->GetFName();
 		TunnelInstanceParams.ObservedObject = TunnelInstance;
-		TunnelInstanceParams.DisplayName = TunnelInstance->GetNodeTitle(ENodeTitleType::ListView);
+		TunnelInstanceParams.DisplayName = TunnelInstance->GetNodeTitle(ENodeTitleType::MenuTitle);
 		TunnelInstanceParams.Tooltip = LOCTEXT("NavigateToNodeLocationHyperlink_ToolTip", "Navigate to the Node");
 		TunnelInstanceParams.NodeFlags = EScriptExecutionNodeFlags::TunnelInstance;
 		TunnelInstanceParams.IconColor = FLinearColor(1.f, 1.f, 1.f, 0.8f);
@@ -1914,23 +1962,16 @@ void FBlueprintTunnelInstanceContext::MapTunnelContext(TSharedPtr<FBlueprintFunc
 		{
 			if (UEdGraph* NestedTunnelGraph = FBlueprintFunctionContext::GetGraphFromNode(NestedTunnelInstance.Key, false))
 			{
-				// Locate the tunnel instance context.
-				const FName TunnelInstanceFunctionName = GetTunnelInstanceFunctionName(NestedTunnelInstance.Key);
-				const FName TunnelFunctionName = GetScopedFunctionNameFromNode(NestedTunnelInstance.Key);
-				TSharedPtr<FBlueprintFunctionContext> FunctionContext = BlueprintContext.Pin()->GetFunctionContext(TunnelInstanceFunctionName);
-				if (!FunctionContext.IsValid())
-				{
-					TSharedPtr<FBlueprintTunnelInstanceContext> NestedContext = MakeShareable(new FBlueprintTunnelInstanceContext);
-					NestedContext->InitialiseContextFromGraph(BlueprintContext.Pin(), TunnelInstanceFunctionName, NestedTunnelGraph);
-					NestedContext->SetParentContext(AsShared());
-					NestedContext->MapTunnelContext(NestedTunnelInstance.Value, NestedTunnelInstance.Key);
-					// Register as both the tunnel instance name and tunnel function name for lookups.
-					BlueprintContext.Pin()->AddFunctionContext(TunnelInstanceFunctionName, NestedContext);
-					BlueprintContext.Pin()->AddFunctionContext(TunnelFunctionName, NestedContext);
-					FunctionContext = NestedContext;
-				}
-				AddChildFunctionContext(TunnelInstanceFunctionName, FunctionContext);
-				AddChildFunctionContext(TunnelFunctionName, FunctionContext);
+				// Map nested tunnel instances.
+				const FName ScopedFunctionName = FBlueprintFunctionContext::GetTunnelInstanceFunctionName(NestedTunnelInstance.Key);
+				TSharedPtr<FBlueprintTunnelInstanceContext> NestedContext = MakeShareable(new FBlueprintTunnelInstanceContext);
+				NestedContext->InitialiseContextFromGraph(BlueprintContext.Pin(), ScopedFunctionName, NestedTunnelGraph);
+				NestedContext->SetParentContext(AsShared());
+				TunnelInstances.Push(NestedTunnelInstance.Key);
+				NestedContext->MapTunnelContext(RootFunctionContext.Pin(), NestedTunnelInstance.Value, TunnelInstances);
+				TunnelInstances.Pop();
+				AddChildFunctionContext(NestedTunnelInstance.Key->GetFName(), NestedContext);
+				BlueprintContext.Pin()->AddFunctionContext(ScopedFunctionName, NestedContext);
 			}
 		}
 		// Find tunnel instance entry sites
@@ -1963,7 +2004,7 @@ void FBlueprintTunnelInstanceContext::MapTunnelContext(TSharedPtr<FBlueprintFunc
 				TunnelEntryParams.NodeName = BoundaryName;
 				TunnelEntryParams.ObservedObject = TunnelInstance;
 				TunnelEntryParams.ObservedPin = InstanceEntryPoint.Value;
-				TunnelEntryParams.DisplayName = TunnelInstance->GetNodeTitle(ENodeTitleType::ListView);
+				TunnelEntryParams.DisplayName = TunnelInstance->GetNodeTitle(ENodeTitleType::MenuTitle);
 				TunnelEntryParams.Tooltip = LOCTEXT("NavigateToTunnelInstanceHyperlink_ToolTip", "Navigate to the Tunnel Instance");
 				TunnelEntryParams.NodeFlags = EScriptExecutionNodeFlags::TunnelEntryPinInstance;
 				TunnelEntryParams.IconColor = FLinearColor(1.f, 1.f, 1.f, 0.8f);
@@ -2029,6 +2070,32 @@ void FBlueprintTunnelInstanceContext::MapTunnelContext(TSharedPtr<FBlueprintFunc
 			InputPins.Add(PinToTrace);
 			MapInputPins(PureChainRootNode, InputPins);
 		}
+		// Register graph nodes to context and discover valid script offsets
+		TSet<const UEdGraphNode*> ProcessedNodes;
+		for (auto ExecNode : ExecutionNodes)
+		{
+			if (UEdGraphNode* GraphNode = const_cast<UEdGraphNode*>(ExecNode.Value->GetTypedObservedObject<UEdGraphNode>()))
+			{
+				if (!ProcessedNodes.Contains(GraphNode))
+				{
+					ProcessedNodes.Add(GraphNode);
+					// Register valid code locations.
+					TArray<int32> CodeLocations;
+					BlueprintClass.Get()->GetDebugData().FindAllCodeLocationsFromSourceNode(GraphNode, Function.Get(), CodeLocations);
+					for (const int32 CodeLocation : CodeLocations)
+					{
+						if (!ScriptOffsetToPins.Contains(CodeLocation))
+						{
+							if (IsCodeLocationValidForThisTunnel(CodeLocation))
+							{
+								ScriptOffsetToPins.Add(CodeLocation);
+							}
+						}
+					}
+					RootFunctionContext.Pin()->AddGraphNodeContextPath(GraphNode, FunctionName);
+				}
+			}
+		}
 	}
 }
 
@@ -2210,7 +2277,7 @@ TSharedPtr<FScriptExecutionNode> FBlueprintTunnelInstanceContext::MapPureNodeExe
 				PinParams.ObservedPin = LinkedPin;
 				PinParams.Tooltip = LOCTEXT("NavigateToPinLocationHyperlink_ToolTip", "Navigate to the Pure Pin");
 				FFormatNamedArguments Args;
-				Args.Add("NodeDisplayName", LinkedNode->GetNodeTitle(ENodeTitleType::ListView));
+				Args.Add("NodeDisplayName", LinkedNode->GetNodeTitle(ENodeTitleType::MenuTitle));
 				Args.Add("PinDisplayName", LinkedPin->PinFriendlyName);
 				FText Format = LinkedPin->PinFriendlyName.IsEmpty() ? LOCTEXT("PureNodeDisplay_Text", "{NodeDisplayName}") : LOCTEXT("PurePinDisplay_Text", "{NodeDisplayName} - {PinDisplayName}");
 				PinParams.DisplayName = FText::Format(Format, Args);
@@ -2308,7 +2375,10 @@ void FBlueprintTunnelInstanceContext::MapInputPins(TSharedPtr<FScriptExecutionNo
 							TSharedPtr<FScriptExecutionNode> PureNode = MapPureNodeExecution(LinkedPin);
 							for (int32 i = 0; i < PinScriptCodeOffsets.Num(); ++i)
 							{
-								PureChainRootNode->AddLinkedNode(PinScriptCodeOffsets[i], PureNode);
+								if (IsCodeLocationValidForThisTunnel(PinScriptCodeOffsets[i]))
+								{
+									PureChainRootNode->AddLinkedNode(PinScriptCodeOffsets[i], PureNode);
+								}
 							}
 						}
 					}
@@ -2466,6 +2536,91 @@ TSharedPtr<FScriptExecutionNode> FBlueprintTunnelInstanceContext::MapTunnelBound
 	return TunnelBoundaryNode;
 }
 
+const int32 FBlueprintTunnelInstanceContext::GetCodeLocationFromPin(const UEdGraphPin* Pin)
+{
+	int32 Result = INDEX_NONE;
+	if (const int32* ScriptOffset = ScriptOffsetToPins.FindKey(Pin))
+	{
+		Result = *ScriptOffset;
+	}
+	else
+	{
+		if (BlueprintClass.IsValid())
+		{
+			// Verify the correct code location for this tunnel ( in case of nested tunnels )
+			TArray<int32> CodeLocations;
+			BlueprintClass.Get()->GetDebugData().FindAllCodeLocationsFromSourcePin(Pin, Function.Get(), CodeLocations);
+			if (CodeLocations.Num() == 1)
+			{
+				Result = CodeLocations[0];
+				ScriptOffsetToPins.Add(Result) = Pin;
+			}
+			else if (CodeLocations.Num() > 1)
+			{
+				TArray<UEdGraphNode*> MacrosAtLocation;
+				for (const int32 CodeLocation : CodeLocations)
+				{
+					if (IsCodeLocationValidForThisTunnel(CodeLocation))
+					{
+						// Cache the value, so we don't do this each time.
+						ScriptOffsetToPins.Add(CodeLocation) = Pin;
+						Result = CodeLocation;
+					}
+				}
+			}
+		}
+	}
+	return Result;
+}
+
+bool FBlueprintTunnelInstanceContext::IsCodeLocationValidForThisTunnel(const int32 CodeLocation) const
+{
+	TArray<UEdGraphNode*> TunnelsAtLocation;
+	BlueprintClass.Get()->GetDebugData().FindMacroInstanceNodesFromCodeLocation(Function.Get(), CodeLocation, TunnelsAtLocation);
+	// Try and identify which code location we want.
+	for (auto Macro : NestedTunnelInstanceStack)
+	{
+		if (!TunnelsAtLocation.Contains(Macro))
+		{
+			return false;
+		}
+	}
+	return true;
+}
+
+FName FBlueprintTunnelInstanceContext::GetTunnelInstanceFunctionName(const UEdGraphNode* GraphNode) const
+{
+	FName ScopedFunctionName = NAME_None;
+	if (GraphNode)
+	{
+		UK2Node_Tunnel* MacroInstance = TunnelInstanceNode.Get();
+		if (MacroInstance == GraphNode)
+		{
+			ScopedFunctionName = FunctionName;
+		}
+		else
+		{
+			ScopedFunctionName = FName(*FString::Printf(TEXT("%s::%s"), *FunctionName.ToString(), *GraphNode->GetName()));
+		}
+	}
+	check (ScopedFunctionName != NAME_None);
+	return ScopedFunctionName;
+}
+
+bool FBlueprintTunnelInstanceContext::GetProfilerContextFromScriptOffset(const int32 ScriptOffset, FNodeExecutionContext& ExecContextOut)
+{
+	// Check the script offset is valid for this context
+	if (ScriptOffsetToPins.Contains(ScriptOffset))
+	{
+		if (const UEdGraphNode* GraphNode = GetNodeFromCodeLocation(ScriptOffset))
+		{
+			ExecContextOut.ProfilerNode = GetProfilerDataForNode(GraphNode->GetFName());
+			ExecContextOut.ProfilerContext = AsShared();
+		}
+	}
+	return ExecContextOut.IsValid();
+}
+
 //////////////////////////////////////////////////////////////////////////
 // FScriptEventPlayback
 
@@ -2484,7 +2639,7 @@ bool FScriptEventPlayback::Process(const TArray<FScriptInstrumentedEvent>& Signa
 		CurrentFunctionName = FunctionContext->GetFunctionName();
 		TSharedPtr<FScriptExecutionNode> EventNode = FunctionContext->GetProfilerDataForNode(EventName);
 		// Find Associated graph nodes and submit into a node map for later processing.
-		TMap<const UEdGraphNode*, NodeSignalHelper> CachedNodeInfo;
+		TMap<TSharedPtr<FScriptExecutionNode>, FNodeSignalHelper> CachedNodeInfo;
 		bProcessingSuccess = true;
 		int32 LastEventIdx = StartIdx;
 		const int32 EventStartOffset = SignalData[StartIdx].IsResumeEvent() ? 3 : 1;
@@ -2528,17 +2683,19 @@ bool FScriptEventPlayback::Process(const TArray<FScriptInstrumentedEvent>& Signa
 				FunctionContext = BlueprintContext->GetFunctionContext(CurrentFunctionName);
 				check(FunctionContext.IsValid());
 			}
-			if (const UEdGraphNode* GraphNode = FunctionContext->GetNodeFromCodeLocation(CurrSignal.GetScriptCodeOffset()))
+			// Grab the node execution context ( this is cached under the hood ).
+			FNodeExecutionContext& NodeExecContext = FunctionContext->GetNodeExecutionContext(CurrSignal.GetScriptCodeOffset());
+			if (const UEdGraphNode* GraphNode = NodeExecContext.GraphNode.Get())
 			{
-				NodeSignalHelper& CurrentNodeData = CachedNodeInfo.FindOrAdd(GraphNode);
 				// Initialise the current node context.
+				FNodeSignalHelper& CurrentNodeData = CachedNodeInfo.FindOrAdd(NodeExecContext.ProfilerNode);
 				if (!CurrentNodeData.IsValid())
 				{
-					FunctionContext->GetProfilerContextFromScriptOffset(CurrSignal.GetScriptCodeOffset(), CurrentNodeData.ImpureNode, CurrentNodeData.FunctionContext);
-					check(CurrentNodeData.IsValid());
+					CurrentNodeData.ImpureNode = NodeExecContext.ProfilerNode;
+					CurrentNodeData.FunctionContext = NodeExecContext.ProfilerContext;
 				}
 				// Check for tunnel boundries and process here
-				if (CurrentNodeData.ImpureNode->IsTunnelInstance())
+				if (NodeExecContext.ProfilerNode->IsTunnelInstance())
 				{
 					ProcessTunnelBoundary(CurrentNodeData, CurrSignal);
 					continue;
@@ -2617,10 +2774,7 @@ bool FScriptEventPlayback::Process(const TArray<FScriptInstrumentedEvent>& Signa
 						TSharedPtr<FScriptExecutionTunnelEntry> TunnelEntry = TracePath.GetTunnel();
 						if (TunnelEntry.IsValid())
 						{
-							TSharedPtr<FScriptExecutionTunnelInstance> TunnelInstance = TunnelEntry->GetTunnelInstance();
-							check (TunnelInstance.IsValid());
-							const UEdGraphNode* TunnelNode = TunnelInstance->GetTypedObservedObject<UEdGraphNode>();
-							NodeSignalHelper& TunnelNodeData = CachedNodeInfo.FindOrAdd(TunnelNode);
+							FNodeSignalHelper& TunnelNodeData = CachedNodeInfo.FindOrAdd(TunnelEntry->GetTunnelInstance());
 							FScriptInstrumentedEvent OverrideEvent(CurrSignal);
 							FTracePath TunnelTrace;
 							TracePath.GetTunnelTracePath(TunnelTrace);
@@ -2690,8 +2844,7 @@ bool FScriptEventPlayback::Process(const TArray<FScriptInstrumentedEvent>& Signa
 						TSharedPtr<FScriptExecutionNode> NextLink = CurrentNodeData.ImpureNode->GetLinkedNodeByScriptOffset(ScriptCodeExit);
 						if (NextLink.IsValid() && NextLink->HasCyclicLinkage())
 						{
-							const UEdGraphNode* LinkedGraphNode = Cast<UEdGraphNode>(NextLink->GetObservedObject());
-							if (NodeSignalHelper* LinkedEntry = CachedNodeInfo.Find(LinkedGraphNode))
+							if (FNodeSignalHelper* LinkedEntry = CachedNodeInfo.Find(NextLink))
 							{
 								if (LinkedEntry->AverageTracePaths.Num())
 								{
@@ -2854,7 +3007,7 @@ bool FScriptEventPlayback::Process(const TArray<FScriptInstrumentedEvent>& Signa
 	return bProcessingSuccess;
 }
 
-void FScriptEventPlayback::ProcessTunnelBoundary(NodeSignalHelper& CurrentNodeData, const FScriptInstrumentedEvent& CurrSignal)
+void FScriptEventPlayback::ProcessTunnelBoundary(FNodeSignalHelper& CurrentNodeData, const FScriptInstrumentedEvent& CurrSignal)
 {
 	// Find grab the tunnel instance exec node.
 	TSharedPtr<FScriptExecutionTunnelInstance> TunnelInstance = StaticCastSharedPtr<FScriptExecutionTunnelInstance>(CurrentNodeData.ImpureNode);
@@ -2915,7 +3068,7 @@ void FScriptEventPlayback::ProcessTunnelBoundary(NodeSignalHelper& CurrentNodeDa
 	}
 }
 
-void FScriptEventPlayback::ProcessExecutionSequence(NodeSignalHelper& CurrentNodeData, const FScriptInstrumentedEvent& CurrSignal)
+void FScriptEventPlayback::ProcessExecutionSequence(FNodeSignalHelper& CurrentNodeData, const FScriptInstrumentedEvent& CurrSignal)
 {
 	switch(CurrSignal.GetType())
 	{
diff --git a/Engine/Source/Developer/BlueprintProfiler/Public/ScriptInstrumentationCapture.h b/Engine/Source/Developer/BlueprintProfiler/Public/ScriptInstrumentationCapture.h
--- a/Engine/Source/Developer/BlueprintProfiler/Public/ScriptInstrumentationCapture.h
+++ b/Engine/Source/Developer/BlueprintProfiler/Public/ScriptInstrumentationCapture.h
@@ -13,9 +13,10 @@ public:
 
 	FScriptInstrumentedEvent() {}
 
-	FScriptInstrumentedEvent(EScriptInstrumentation::Type InEventType, const FName InFunctionClassScopeName, const FString& InPathData)
+	FScriptInstrumentedEvent(EScriptInstrumentation::Type InEventType, const FName InFunctionClassScopeName, const FName InFunctionName, const FString& InPathData)
 		: EventType(InEventType)
 		, FunctionClassScopeName(InFunctionClassScopeName)
+		, FunctionName(InFunctionName)
 		, PathData(InPathData)
 		, CodeOffset(-1)
 		, Time(FPlatformTime::Seconds())
@@ -114,18 +115,23 @@ public:
 	/** Reset the current event context */
 	void ResetContext();
 
+	/** Return to previous scope */
+	void PopScope();
+
 private:
 
-	/** The current class/blueprint context */
-	TWeakObjectPtr<const UClass> ContextClass;
-	/** The current function class scope */
-	TWeakObjectPtr<const UClass> ContextFunctionClassScope;
-	/** The current instance context */
-	TWeakObjectPtr<const UObject> ContextObject;
-	/** The current event context */
-	FName ActiveEvent;
+	struct FContextScope
+	{
+		/** The current class/blueprint context */
+		TWeakObjectPtr<const UClass> ContextClass;
+		/** The current function class scope */
+		TWeakObjectPtr<const UClass> ContextFunctionClassScope;
+		/** The current instance context */
+		TWeakObjectPtr<const UObject> ContextObject;
+	};
+
 	/** The current pending inline Event. NAME_None if no event is pending */
 	FName PendingInlineEventName;
 	/** The Class Scope Stack for traversing through parent blueprint calls */
-	TArray<FName> ScopeStack;
+	TArray<FContextScope> ScopeStack;
 };
diff --git a/Engine/Source/Developer/BlueprintProfiler/Public/ScriptInstrumentationPlayback.h b/Engine/Source/Developer/BlueprintProfiler/Public/ScriptInstrumentationPlayback.h
--- a/Engine/Source/Developer/BlueprintProfiler/Public/ScriptInstrumentationPlayback.h
+++ b/Engine/Source/Developer/BlueprintProfiler/Public/ScriptInstrumentationPlayback.h
@@ -168,6 +168,18 @@ private:
 	TSimpleRingBuffer<FBlueprintExecutionTrace> ExecutionTraceHistory;
 };
 
+//////////////////////////////////////////////////////////////////////////
+// FNodeExecutionContext
+
+struct FNodeExecutionContext
+{
+	bool IsValid() const { return ProfilerNode.IsValid() && ProfilerContext.IsValid(); }
+
+	TSharedPtr<FScriptExecutionNode> ProfilerNode;
+	TSharedPtr<FBlueprintFunctionContext> ProfilerContext;
+	TWeakObjectPtr<const UEdGraphNode> GraphNode;
+};
+
 //////////////////////////////////////////////////////////////////////////
 // FBlueprintFunctionContext
 
@@ -208,11 +220,14 @@ public:
 	const UEdGraphPin* GetPinFromCodeLocation(const int32 ScriptOffset);
 
 	/** Returns the script offset for the specified pin */
-	const int32 GetCodeLocationFromPin(const UEdGraphPin* Pin) const;
+	virtual const int32 GetCodeLocationFromPin(const UEdGraphPin* Pin);
 
 	/** Returns all registered script offsets for the specified pin */
 	void GetAllCodeLocationsFromPin(const UEdGraphPin* Pin, TArray<int32>& OutCodeLocations) const;
 
+	/** Returns all tunnels active at the code location */
+	void FindAllTunnelsAtCodeLocation(const int32 CodeLocation, TArray<UEdGraphNode*>& MacrosAtLocation) const;
+
 	/** Returns the pure node script code range for the specified node */
 	FInt32Range GetPureNodeScriptCodeRange(const UEdGraphNode* Node) const;
 
@@ -238,7 +253,7 @@ public:
 	template<typename ExecNodeType> TSharedPtr<ExecNodeType> GetTypedProfilerDataForGraphNode(const UEdGraphNode* Node);
 
 	/** Returns true if the execution node and function context can be located and initialised from the script offset */
-	virtual bool GetProfilerContextFromScriptOffset(const int32 ScriptOffset, TSharedPtr<FScriptExecutionNode>& ExecNode, TSharedPtr<FBlueprintFunctionContext>& FunctionContext);
+	virtual bool GetProfilerContextFromScriptOffset(const int32 ScriptOffset, FNodeExecutionContext& ExecContextOut);
 
 	/** Adds all function entry points to node as children */
 	void AddCallSiteEntryPointsToNode(TSharedPtr<FScriptExecutionNode> CallingNode) const;
@@ -277,7 +292,7 @@ public:
 	static FName GetGraphNameFromNode(const UEdGraphNode* GraphNode);
 
 	/** Utility function that returns the tunnel instance function name */
-	static FName GetTunnelInstanceFunctionName(const UEdGraphNode* GraphNode);
+	virtual FName GetTunnelInstanceFunctionName(const UEdGraphNode* GraphNode) const;
 
 	/** Utility to get the scoped function name from the graph node */
 	FName GetScopedFunctionNameFromNode(const UEdGraphNode* GraphNode) const;
@@ -288,6 +303,9 @@ public:
 	/** Utility function that returns customizations for a node (e.g. Node Color/Icon) */
 	void GetNodeCustomizations(FScriptExecNodeParams& ParamsInOut) const;
 
+	/** Lookup node execution context using the script offset - implements caching internally */
+	virtual FNodeExecutionContext& GetNodeExecutionContext(const int32 ScriptOffset);
+
 protected:
 
 	friend FBlueprintExecutionContext;
@@ -351,6 +369,12 @@ protected:
 	/** Maps execution pin tunnel entry */
 	void MapTunnelExits(TSharedPtr<FScriptExecutionTunnelEntry> TunnelEntryPoint);
 
+	/** Register a function context name to a particular graph node for discovery during playback. */
+	void AddGraphNodeContextPath(const UEdGraphNode* GraphNode, const FName ScopedTunnelContextName)
+	{
+		GraphNodeToContextName.Add(GraphNode, ScopedTunnelContextName);
+	}
+
 protected:
 
 	/** The graph name */
@@ -369,6 +393,10 @@ protected:
 	TWeakObjectPtr<UBlueprintGeneratedClass> BlueprintClass;
 	/** ScriptCode offset to node cache */
 	TMap<int32, TWeakObjectPtr<const UEdGraphNode>> ScriptOffsetToNodes;
+	/** ScriptCode offset to node context */
+	TMap<int32, FNodeExecutionContext> ScriptOffsetToNodeContext;
+	/** GraphNode to tunnel context names */
+	TMultiMap<TWeakObjectPtr<const UEdGraphNode>, FName> GraphNodeToContextName;
 	/** ScriptCode offset to tunnel node cache */
 	TMap<TWeakObjectPtr<const UEdGraphNode>, TWeakObjectPtr<const UEdGraphNode>> NodeToTunnelNode;
 	/** ScriptCode offset to pin cache */
@@ -393,10 +421,12 @@ public:
 	// ~FBlueprintFunctionContext Begin
 	virtual void InitialiseContextFromGraph(TSharedPtr<FBlueprintExecutionContext> BlueprintContext, const FName TunnelInstanceName, UEdGraph* TunnelGraph);
 	virtual void MapFunction() override {}
+	virtual const int32 GetCodeLocationFromPin(const UEdGraphPin* Pin) override;
+	virtual bool GetProfilerContextFromScriptOffset(const int32 ScriptOffset, FNodeExecutionContext& ExecContextOut) override;
 	// ~FBlueprintFunctionContext End
 
 	/** Maps a tunnel function context */
-	void MapTunnelContext(TSharedPtr<FBlueprintFunctionContext> CallingFunctionContext, UK2Node_Tunnel* TunnelInstance);
+	void MapTunnelContext(TSharedPtr<FBlueprintFunctionContext> RootFunctionContextIn, TSharedPtr<FBlueprintFunctionContext> CallingFunctionContext, TArray<UK2Node_Tunnel*> TunnelInstances);
 
 	/** Sets the tunnel context parent context */
 	void SetParentContext(TSharedPtr<FBlueprintFunctionContext> ParentContextIn);
@@ -411,6 +441,7 @@ private:
 	virtual TSharedPtr<FScriptExecutionNode> MapPureNodeExecution(const UEdGraphPin* LinkedPin) override;
 	virtual TSharedPtr<FScriptExecutionNode> GetTunnelBoundaryNode(const UEdGraphPin* TunnelPin) override;
 	virtual TSharedPtr<FScriptExecutionNode> MapTunnelBoundary(const UEdGraphPin* TunnelPin) override;
+	virtual FName GetTunnelInstanceFunctionName(const UEdGraphNode* GraphNode) const override;
 	// ~FBlueprintFunctionContext End
 	
 	/** Maps input and output pins in a tunnel graph */
@@ -431,12 +462,19 @@ private:
 	/** Looks up a function context using the nodes graph */
 	TSharedPtr<FBlueprintFunctionContext> FindContextByNodeGraph(const UEdGraphNode* GraphNode);
 
+	/** Determines if a code location is valid for this tunnel context instance using the active tunnel hierarchy */
+	bool IsCodeLocationValidForThisTunnel(const int32 CodeLocation) const;
+
 private:
 
+	/** Root function context */
+	TWeakPtr<FBlueprintFunctionContext> RootFunctionContext;
 	/** Tunnel context that are contained by this tunnel context */
 	TWeakPtr<FBlueprintTunnelInstanceContext> ParentTunnel;
 	/** The tunnel instance this context represents */
 	TWeakObjectPtr<UK2Node_Tunnel> TunnelInstanceNode;
+	/** The nested tunnel instance stack */
+	TArray<UK2Node_Tunnel*> NestedTunnelInstanceStack;
 	/** The tunnel pure pin link nodes */
 	TMap<FEdGraphPinReference, TSharedPtr<FScriptExecutionNode>> PureLinkNodes;
 	/** Staging entry point */
@@ -451,12 +489,13 @@ class FScriptEventPlayback : public TSharedFromThis<FScriptEventPlayback>
 {
 public:
 
-	struct NodeSignalHelper
+	struct FNodeSignalHelper
 	{
 		/** Easy validation */
 		bool IsValid() const { return ImpureNode.IsValid() && FunctionContext.IsValid(); }
 
 		/** Node signal properties */
+		TWeakObjectPtr<const UEdGraphNode> GraphNode;
 		TSharedPtr<FScriptExecutionNode> ImpureNode;
 		TMap<int32, TSharedPtr<FScriptExecutionNode>> PureNodes;
 		TSharedPtr<FBlueprintFunctionContext> FunctionContext;
@@ -467,7 +506,7 @@ public:
 		TArray<FTracePath> InclusiveTracePaths;
 	};
 
-	struct TunnelEventHelper
+	struct FTunnelEventHelper
 	{
 		double TunnelEntryTime;
 		FTracePath EntryTracePath;
@@ -509,10 +548,10 @@ public:
 private:
 
 	/** Process tunel boundries */
-	void ProcessTunnelBoundary(NodeSignalHelper& CurrentNodeData, const FScriptInstrumentedEvent& CurrSignal);
+	void ProcessTunnelBoundary(FNodeSignalHelper& CurrentNodeData, const FScriptInstrumentedEvent& CurrSignal);
 
 	/** Process execution sequence */
-	void ProcessExecutionSequence(NodeSignalHelper& CurrentNodeData, const FScriptInstrumentedEvent& CurrSignal);
+	void ProcessExecutionSequence(FNodeSignalHelper& CurrentNodeData, const FScriptInstrumentedEvent& CurrSignal);
 
 	/** Add to trace history */
 	void AddToTraceHistory(const TSharedPtr<FScriptExecutionNode> ProfilerNode, const FScriptInstrumentedEvent& TraceSignal);
@@ -537,8 +576,6 @@ private:
 	TArray<FTracePath> TraceStack;
 	/** Current tunnel tracepath stack */
 	TArray<FTracePath> TunnelTraceStack;
-	/** Active tunnels */
-	TMap<FName, TunnelEventHelper> ActiveTunnels;
 	/** Event Timings */
 	TMultiMap<FName, double> EventTimings;
 	/** Blueprint exec context */
diff --git a/Engine/Source/Editor/AnimGraph/Classes/K2Node_TransitionRuleGetter.h b/Engine/Source/Editor/AnimGraph/Classes/K2Node_TransitionRuleGetter.h
--- a/Engine/Source/Editor/AnimGraph/Classes/K2Node_TransitionRuleGetter.h
+++ b/Engine/Source/Editor/AnimGraph/Classes/K2Node_TransitionRuleGetter.h
@@ -50,6 +50,7 @@ class UK2Node_TransitionRuleGetter : public UK2Node
 
 	// UK2Node interface
 	virtual void GetMenuActions(FBlueprintActionDatabaseRegistrar& ActionRegistrar) const override;
+	virtual bool IsNodePure() const override { return true; }
 	// end of UK2Node interface
 
 	// @todo document
diff --git a/Engine/Source/Editor/BlueprintGraph/Classes/K2Node.h b/Engine/Source/Editor/BlueprintGraph/Classes/K2Node.h
--- a/Engine/Source/Editor/BlueprintGraph/Classes/K2Node.h
+++ b/Engine/Source/Editor/BlueprintGraph/Classes/K2Node.h
@@ -60,6 +60,9 @@ struct FOptionalPinFromProperty
 	UPROPERTY(EditAnywhere, Category=OptionalPin)
 	bool bHasOverridePin;
 
+	UPROPERTY(EditAnywhere, Category=OptionalPin)
+	bool bIsMarkedForAdvancedDisplay;
+
 	/** TRUE if the override value is enabled for use */
 	UPROPERTY(EditAnywhere, Category = OptionalPin)
 	bool bIsOverrideEnabled;
@@ -73,7 +76,8 @@ struct FOptionalPinFromProperty
 	bool bIsOverridePinVisible;
 
 	FOptionalPinFromProperty()
-		: bIsOverrideEnabled(true)
+		: bIsMarkedForAdvancedDisplay(false)
+		, bIsOverrideEnabled(true)
 		, bIsSetValuePinVisible(true)
 		, bIsOverridePinVisible(true)
 	{
@@ -90,6 +94,7 @@ struct FOptionalPinFromProperty
 		, bPropertyIsCustomized(bInPropertyIsCustomized)
 		, CategoryName(InCategoryName)
 		, bHasOverridePin(bInHasOverridePin)
+		, bIsMarkedForAdvancedDisplay(false)
 		, bIsOverrideEnabled(true)
 		, bIsSetValuePinVisible(true)
 		, bIsOverridePinVisible(true)
@@ -296,9 +301,6 @@ class UK2Node : public UEdGraphNode
 	/** This function if used for nodes that needs CDO for validation (Called before expansion)*/
 	BLUEPRINTGRAPH_API virtual void EarlyValidation(class FCompilerResultsLog& MessageLog) const {}
 
-	/** This function if used for nodes that should validate after expansion */
-	BLUEPRINTGRAPH_API virtual void ValidateNodeAfterPrune(class FCompilerResultsLog& MessageLog) const {}
-
 	/** This function returns an arbitrary number of attributes that describe this node for analytics events */
 	BLUEPRINTGRAPH_API virtual void GetNodeAttributes( TArray<TKeyValuePair<FString, FString>>& OutNodeAttributes ) const;
 
diff --git a/Engine/Source/Editor/BlueprintGraph/Classes/K2Node_CallFunction.h b/Engine/Source/Editor/BlueprintGraph/Classes/K2Node_CallFunction.h
--- a/Engine/Source/Editor/BlueprintGraph/Classes/K2Node_CallFunction.h
+++ b/Engine/Source/Editor/BlueprintGraph/Classes/K2Node_CallFunction.h
@@ -91,7 +91,7 @@ public:
 	virtual bool ShouldDrawAsBead() const override;
 	virtual FText GetCompactNodeTitle() const override;
 	virtual void PostPasteNode() override;
-	virtual void ValidateNodeAfterPrune(class FCompilerResultsLog& MessageLog) const override;
+	virtual void ValidateNodeDuringCompilation(class FCompilerResultsLog& MessageLog) const override;
 	virtual bool ShouldShowNodeProperties() const override;
 	virtual void GetRedirectPinNames(const UEdGraphPin& Pin, TArray<FString>& RedirectPinNames) const override;
 	virtual void NotifyPinConnectionListChanged(UEdGraphPin* Pin) override;
diff --git a/Engine/Source/Editor/BlueprintGraph/Classes/K2Node_CustomEvent.h b/Engine/Source/Editor/BlueprintGraph/Classes/K2Node_CustomEvent.h
--- a/Engine/Source/Editor/BlueprintGraph/Classes/K2Node_CustomEvent.h
+++ b/Engine/Source/Editor/BlueprintGraph/Classes/K2Node_CustomEvent.h
@@ -32,6 +32,7 @@ class UK2Node_CustomEvent : public UK2Node_Event
 	virtual FSlateIcon GetIconAndTint(FLinearColor& OutColor) const override;
 	virtual void AutowireNewNode(UEdGraphPin* FromPin) override;
 	virtual void AddSearchMetaDataInfo(TArray<struct FSearchTagDataPair>& OutTaggedMetaData) const override;
+	virtual FText GetKeywords() const override;
 	//~ End UEdGraphNode Interface
 
 	//~ Begin UK2Node Interface
diff --git a/Engine/Source/Editor/BlueprintGraph/Classes/K2Node_FunctionTerminator.h b/Engine/Source/Editor/BlueprintGraph/Classes/K2Node_FunctionTerminator.h
--- a/Engine/Source/Editor/BlueprintGraph/Classes/K2Node_FunctionTerminator.h
+++ b/Engine/Source/Editor/BlueprintGraph/Classes/K2Node_FunctionTerminator.h
@@ -26,6 +26,7 @@ class UK2Node_FunctionTerminator : public UK2Node_EditablePinBase
 	virtual bool CanDuplicateNode() const override { return false; }
 	virtual FLinearColor GetNodeTitleColor() const override;
 	virtual FString CreateUniquePinName(FString SourcePinName) const override;
+	virtual void ValidateNodeDuringCompilation(class FCompilerResultsLog& MessageLog) const override;
 	//~ End UEdGraphNode Interface
 
 	//~ Begin UK2Node Interface
diff --git a/Engine/Source/Editor/BlueprintGraph/Classes/K2Node_Knot.h b/Engine/Source/Editor/BlueprintGraph/Classes/K2Node_Knot.h
--- a/Engine/Source/Editor/BlueprintGraph/Classes/K2Node_Knot.h
+++ b/Engine/Source/Editor/BlueprintGraph/Classes/K2Node_Knot.h
@@ -30,7 +30,7 @@ public:
 	virtual void PostReconstructNode() override;
 	virtual int32 GetNodeRefreshPriority() const override { return EBaseNodeRefreshPriority::Low_UsesDependentWildcard; }
 	virtual void GetMenuActions(FBlueprintActionDatabaseRegistrar& ActionRegistrar) const override;
-	virtual bool IsNodePure() const { return true; }
+	virtual bool IsNodePure() const override { return true; }
 	// End of UK2Node interface
 
 	UEdGraphPin* GetInputPin() const
diff --git a/Engine/Source/Editor/BlueprintGraph/Classes/K2Node_MakeStruct.h b/Engine/Source/Editor/BlueprintGraph/Classes/K2Node_MakeStruct.h
--- a/Engine/Source/Editor/BlueprintGraph/Classes/K2Node_MakeStruct.h
+++ b/Engine/Source/Editor/BlueprintGraph/Classes/K2Node_MakeStruct.h
@@ -62,9 +62,15 @@ protected:
 		const uint8* const SampleStructMemory;
 	public:
 		FMakeStructPinManager(const uint8* InSampleStructMemory);
+
+		bool HasAdvancedPins() const { return bHasAdvancedPins; }
 	protected:
+		virtual void GetRecordDefaults(UProperty* TestProperty, FOptionalPinFromProperty& Record) const override;
 		virtual void CustomizePinData(UEdGraphPin* Pin, FName SourcePropertyName, int32 ArrayIndex, UProperty* Property) const override;
 		virtual bool CanTreatPropertyAsOptional(UProperty* TestProperty) const override;
+
+		/** set by GetRecordDefaults(), mutable as it is a const function */
+		mutable bool bHasAdvancedPins;
 	};
 
 private:
diff --git a/Engine/Source/Editor/BlueprintGraph/Classes/K2Node_TunnelBoundary.h b/Engine/Source/Editor/BlueprintGraph/Classes/K2Node_TunnelBoundary.h
--- a/Engine/Source/Editor/BlueprintGraph/Classes/K2Node_TunnelBoundary.h
+++ b/Engine/Source/Editor/BlueprintGraph/Classes/K2Node_TunnelBoundary.h
@@ -17,9 +17,9 @@ class BLUEPRINTGRAPH_API UK2Node_TunnelBoundary : public UK2Node
 {
 	GENERATED_UCLASS_BODY()
 
-	/** Tunnel Graph Name */
+	/** Base Name */
 	UPROPERTY()
-	FName TunnelGraphName;
+	FName BaseName;
 
 	/** Final exit site */
 	UPROPERTY()
@@ -41,6 +41,9 @@ public:
 	/** Creates Boundary nodes for a tunnel instance */
 	static void CreateBoundaryNodesForTunnelInstance(UK2Node_Tunnel* TunnelInstance, UEdGraph* TunnelGraph, class FCompilerResultsLog& MessageLog);
 
+	/** Creates Boundaries for expansion nodes */
+	static void CreateBoundariesForExpansionNodes(UEdGraphNode* SourceNode, TArray<UEdGraphNode*>& ExpansionNodes, TMap<UEdGraphPin*, UEdGraphPin*>& LinkedPinMap, FCompilerResultsLog& MessageLog);
+
 	/** Checks if the tunnel node is a pure tunnel rather than a tunnel instance */
 	static bool IsPureTunnel(UK2Node_Tunnel* Tunnel);
 
@@ -52,8 +55,14 @@ protected:
 	/** Wires up the tunnel exit boundries */
 	void WireUpTunnelExit(UK2Node_Tunnel* TunnelInstance, UEdGraphPin* TunnelPin, FCompilerResultsLog& MessageLog);
 
-	/** Determines the tunnel graph name */
-	void DetermineTunnelGraphName(UK2Node_Tunnel* TunnelInstance);
+	/** Wires up the entry boundries */
+	void WireUpEntry(UEdGraphNode* SourceNode, UEdGraphPin* SourcePin, TArray<UEdGraphPin*>& EntryPins, FCompilerResultsLog& MessageLog);
+
+	/** Wires up the exit boundries */
+	void WireUpExit(UEdGraphNode* SourceNode, UEdGraphPin* SourcePin, TArray<UEdGraphPin*>& ExitPins, FCompilerResultsLog& MessageLog);
+
+	/** Create base node name */
+	void CreateBaseNodeName(UEdGraphNode* SourceNode);
 
 	/** Build a guid map for nodes from the original graph */
 	static void BuildSourceNodeMap(UEdGraphNode* Tunnel, TMap<FGuid, const UEdGraphNode*>& SourceNodeMap);
diff --git a/Engine/Source/Editor/BlueprintGraph/Private/BlueprintActionDatabase.cpp b/Engine/Source/Editor/BlueprintGraph/Private/BlueprintActionDatabase.cpp
--- a/Engine/Source/Editor/BlueprintGraph/Private/BlueprintActionDatabase.cpp
+++ b/Engine/Source/Editor/BlueprintGraph/Private/BlueprintActionDatabase.cpp
@@ -1304,7 +1304,7 @@ void FBlueprintActionDatabase::RefreshAssetActions(UObject* const AssetObject)
 		for( auto Level : WorldAsset->GetLevels() )
 		{
 			UBlueprint* LevelScript = Cast<UBlueprint>(Level->GetLevelScriptBlueprint(true));
-			if (Level->bIsVisible && LevelScript)
+			if (LevelScript != nullptr)
 			{
 				AddBlueprintGraphActions(LevelScript, AssetActionList);
 				if (UClass* SkeletonClass = LevelScript->SkeletonGeneratedClass)
diff --git a/Engine/Source/Editor/BlueprintGraph/Private/CallFunctionHandler.cpp b/Engine/Source/Editor/BlueprintGraph/Private/CallFunctionHandler.cpp
--- a/Engine/Source/Editor/BlueprintGraph/Private/CallFunctionHandler.cpp
+++ b/Engine/Source/Editor/BlueprintGraph/Private/CallFunctionHandler.cpp
@@ -412,17 +412,16 @@ void FKCHandler_CallFunction::CreateFunctionCallStatement(FKismetFunctionContext
 			}
 
 			bool bInlineEventCall = false;
-			bool bLocalCall = false;
+			bool bEmitInstrumentPushState = false;
 			FName EventName = NAME_None;
 			if (Context.IsInstrumentationRequired())
 			{
 				if (UK2Node_CallFunction* CallFunctionNode = Cast<UK2Node_CallFunction>(Node))
 				{
-					if (CallFunctionNode->FunctionReference.IsSelfContext())
+					const UEdGraphNode* EventNodeOut = nullptr;
+					if (UEdGraph* FunctionGraph = CallFunctionNode->GetFunctionGraph(EventNodeOut))
 					{
-						bLocalCall = true;
-						const UEdGraphNode* EventNodeOut = nullptr;
-						CallFunctionNode->GetFunctionGraph(EventNodeOut);
+						bEmitInstrumentPushState = true;
 						if (const UK2Node_Event* EventNode = Cast<UK2Node_Event>(EventNodeOut))
 						{
 							bInlineEventCall = true;
@@ -431,7 +430,7 @@ void FKCHandler_CallFunction::CreateFunctionCallStatement(FKismetFunctionContext
 					}
 				}
 			}
-			const bool bInstrumentFunctionEntry = Context.IsInstrumentationRequired() && (bIsLatent || bInlineEventCall || bLocalCall);
+			const bool bInstrumentFunctionEntry = Context.IsInstrumentationRequired() && (bIsLatent || bInlineEventCall || bEmitInstrumentPushState);
 			if (bInstrumentFunctionEntry)
 			{
 				if (bInlineEventCall)
diff --git a/Engine/Source/Editor/BlueprintGraph/Private/EdGraphSchema_K2.cpp b/Engine/Source/Editor/BlueprintGraph/Private/EdGraphSchema_K2.cpp
--- a/Engine/Source/Editor/BlueprintGraph/Private/EdGraphSchema_K2.cpp
+++ b/Engine/Source/Editor/BlueprintGraph/Private/EdGraphSchema_K2.cpp
@@ -752,7 +752,7 @@ bool UEdGraphSchema_K2::CanFunctionBeUsedInGraph(const UClass* InClass, const UF
 					auto BP = FBlueprintEditorUtils::FindBlueprintForGraph(InDestGraph);
 					const bool bIsFunctLib = BP && (EBlueprintType::BPTYPE_FunctionLibrary == BP->BlueprintType);
 					UClass* ParentClass = BP ? BP->ParentClass : NULL;
-					const bool bIncompatibleParrent = ParentClass && (!ParentClass->GetDefaultObject()->ImplementsGetWorld() && !ParentClass->HasMetaData(FBlueprintMetadata::MD_ShowWorldContextPin));
+					const bool bIncompatibleParrent = ParentClass && (!ParentClass->GetDefaultObject()->ImplementsGetWorld() && !ParentClass->HasMetaDataHierarchical(FBlueprintMetadata::MD_ShowWorldContextPin));
 					if (!bIsFunctLib && bIncompatibleParrent)
 					{
 						if (OutReason != nullptr)
diff --git a/Engine/Source/Editor/BlueprintGraph/Private/K2Node.cpp b/Engine/Source/Editor/BlueprintGraph/Private/K2Node.cpp
--- a/Engine/Source/Editor/BlueprintGraph/Private/K2Node.cpp
+++ b/Engine/Source/Editor/BlueprintGraph/Private/K2Node.cpp
@@ -971,6 +971,7 @@ void FOptionalPinManager::RebuildProperty(UProperty* TestProperty, FName Categor
 
 	bool bNegate = false;
 	Record->bHasOverridePin = PropertyCustomizationHelpers::GetEditConditionProperty(TestProperty, bNegate) != nullptr;
+	Record->bIsMarkedForAdvancedDisplay = TestProperty->HasAnyPropertyFlags(CPF_AdvancedDisplay);
 
 	// Get the defaults
 	GetRecordDefaults(TestProperty, *Record);
diff --git a/Engine/Source/Editor/BlueprintGraph/Private/K2Node_CallFunction.cpp b/Engine/Source/Editor/BlueprintGraph/Private/K2Node_CallFunction.cpp
--- a/Engine/Source/Editor/BlueprintGraph/Private/K2Node_CallFunction.cpp
+++ b/Engine/Source/Editor/BlueprintGraph/Private/K2Node_CallFunction.cpp
@@ -867,7 +867,7 @@ bool UK2Node_CallFunction::CreatePinsForFunctionCall(const UFunction* Function)
 	TSet<FString> InternalPins;
 	FBlueprintEditorUtils::GetHiddenPinsForFunction(Graph, Function, PinsToHide, &InternalPins);
 
-	const bool bShowWorldContextPin = ((PinsToHide.Num() > 0) && BP && BP->ParentClass && BP->ParentClass->HasMetaData(FBlueprintMetadata::MD_ShowWorldContextPin));
+	const bool bShowWorldContextPin = ((PinsToHide.Num() > 0) && BP && BP->ParentClass && BP->ParentClass->HasMetaDataHierarchical(FBlueprintMetadata::MD_ShowWorldContextPin));
 
 	// Create the inputs and outputs
 	bool bAllPinsGood = true;
@@ -1714,7 +1714,7 @@ void UK2Node_CallFunction::PostPasteNode()
 		TSet<FString> PinsToHide;
 		FBlueprintEditorUtils::GetHiddenPinsForFunction(GetGraph(), Function, PinsToHide);
 
-		const bool bShowWorldContextPin = ((PinsToHide.Num() > 0) && GetBlueprint()->ParentClass->HasMetaData(FBlueprintMetadata::MD_ShowWorldContextPin));
+		const bool bShowWorldContextPin = ((PinsToHide.Num() > 0) && GetBlueprint()->ParentClass->HasMetaDataHierarchical(FBlueprintMetadata::MD_ShowWorldContextPin));
 
 		FString const DefaultToSelfMetaValue = Function->GetMetaData(FBlueprintMetadata::MD_DefaultToSelf);
 		FString const WorldContextMetaValue  = Function->GetMetaData(FBlueprintMetadata::MD_WorldContext);
@@ -1746,9 +1746,9 @@ void UK2Node_CallFunction::PostDuplicate(bool bDuplicateForPIE)
 	}
 }
 
-void UK2Node_CallFunction::ValidateNodeAfterPrune(class FCompilerResultsLog& MessageLog) const
+void UK2Node_CallFunction::ValidateNodeDuringCompilation(class FCompilerResultsLog& MessageLog) const
 {
-	Super::ValidateNodeAfterPrune(MessageLog);
+	Super::ValidateNodeDuringCompilation(MessageLog);
 
 	const UBlueprint* Blueprint = GetBlueprint();
 	UFunction *Function = GetTargetFunction();
@@ -1818,7 +1818,7 @@ void UK2Node_CallFunction::ValidateNodeAfterPrune(class FCompilerResultsLog& Mes
 			check(Blueprint);
 			UClass* ParentClass = Blueprint->ParentClass;
 			check(ParentClass);
-			if (ParentClass && !ParentClass->GetDefaultObject()->ImplementsGetWorld() && !ParentClass->HasMetaData(FBlueprintMetadata::MD_ShowWorldContextPin))
+			if (ParentClass && !ParentClass->GetDefaultObject()->ImplementsGetWorld() && !ParentClass->HasMetaDataHierarchical(FBlueprintMetadata::MD_ShowWorldContextPin))
 			{
 				MessageLog.Warning(*LOCTEXT("FunctionUnsafeInContext", "Function '@@' is unsafe to call from blueprints of class '@@'.").ToString(), this, ParentClass);
 			}
@@ -1826,6 +1826,17 @@ void UK2Node_CallFunction::ValidateNodeAfterPrune(class FCompilerResultsLog& Mes
 	}
 
 	FDynamicOutputHelper::VerifyNode(this, MessageLog);
+
+	for (UEdGraphPin* Pin : Pins)
+	{
+		if (Pin && Pin->PinType.bIsWeakPointer && !Pin->PinType.IsContainer())
+		{
+			const FString ErrorString = FString::Printf(
+				*LOCTEXT("WeakPtrNotSupportedError", "Weak prointer is not supported as function parameter. Pin '%s' @@").ToString(),
+				*Pin->GetName());
+			MessageLog.Error(*ErrorString, this);
+		}
+	}
 }
 
 void UK2Node_CallFunction::Serialize(FArchive& Ar)
diff --git a/Engine/Source/Editor/BlueprintGraph/Private/K2Node_ConstructObjectFromClass.cpp b/Engine/Source/Editor/BlueprintGraph/Private/K2Node_ConstructObjectFromClass.cpp
--- a/Engine/Source/Editor/BlueprintGraph/Private/K2Node_ConstructObjectFromClass.cpp
+++ b/Engine/Source/Editor/BlueprintGraph/Private/K2Node_ConstructObjectFromClass.cpp
@@ -35,7 +35,7 @@ bool UK2Node_ConstructObjectFromClass::UseWorldContext() const
 {
 	auto BP = GetBlueprint();
 	const UClass* ParentClass = BP ? BP->ParentClass : nullptr;
-	return ParentClass ? ParentClass->HasMetaData(FBlueprintMetadata::MD_ShowWorldContextPin) : false;
+	return ParentClass ? ParentClass->HasMetaDataHierarchical(FBlueprintMetadata::MD_ShowWorldContextPin) != nullptr : false;
 }
 
 void UK2Node_ConstructObjectFromClass::AllocateDefaultPins()
diff --git a/Engine/Source/Editor/BlueprintGraph/Private/K2Node_CustomEvent.cpp b/Engine/Source/Editor/BlueprintGraph/Private/K2Node_CustomEvent.cpp
--- a/Engine/Source/Editor/BlueprintGraph/Private/K2Node_CustomEvent.cpp
+++ b/Engine/Source/Editor/BlueprintGraph/Private/K2Node_CustomEvent.cpp
@@ -527,4 +527,13 @@ void UK2Node_CustomEvent::AddSearchMetaDataInfo(TArray<struct FSearchTagDataPair
 	OutTaggedMetaData.Add(FSearchTagDataPair(FFindInBlueprintSearchTags::FiB_NativeName, FText::FromName(CustomFunctionName)));
 }
 
+FText UK2Node_CustomEvent::GetKeywords() const
+{
+	FText ParentKeywords = Super::GetKeywords();
+
+	FFormatNamedArguments Args;
+	Args.Add(TEXT("ParentKeywords"), ParentKeywords);
+	return FText::Format(LOCTEXT("CustomEventKeywords", "{ParentKeywords} Custom"), Args);
+}
+
 #undef LOCTEXT_NAMESPACE
diff --git a/Engine/Source/Editor/BlueprintGraph/Private/K2Node_FunctionResult.cpp b/Engine/Source/Editor/BlueprintGraph/Private/K2Node_FunctionResult.cpp
--- a/Engine/Source/Editor/BlueprintGraph/Private/K2Node_FunctionResult.cpp
+++ b/Engine/Source/Editor/BlueprintGraph/Private/K2Node_FunctionResult.cpp
@@ -93,7 +93,7 @@ public:
 		{
 			GenerateAssigments(Context, Node);
 
-			if (Context.IsDebuggingOrInstrumentationRequired() && Node)
+			if (Context.IsDebuggingOrInstrumentationRequired() && Node && !Context.IsExpansionNode(Node))
 			{
 				FBlueprintCompiledStatement& TraceStatement = Context.AppendStatementForNode(Node);
 				TraceStatement.Type = Context.GetWireTraceType();
diff --git a/Engine/Source/Editor/BlueprintGraph/Private/K2Node_FunctionTerminator.cpp b/Engine/Source/Editor/BlueprintGraph/Private/K2Node_FunctionTerminator.cpp
--- a/Engine/Source/Editor/BlueprintGraph/Private/K2Node_FunctionTerminator.cpp
+++ b/Engine/Source/Editor/BlueprintGraph/Private/K2Node_FunctionTerminator.cpp
@@ -4,6 +4,7 @@
 #include "BlueprintGraphPrivatePCH.h"
 #include "K2Node_FunctionTerminator.h"
 #include "GraphEditorSettings.h"
+#include "CompilerResultsLog.h"
 
 #define LOCTEXT_NAMESPACE "K2Node"
 
@@ -103,4 +104,20 @@ void UK2Node_FunctionTerminator::PromoteFromInterfaceOverride(bool bIsPrimaryTer
 	Schema->ReconstructNode(*this, true);
 }
 
+void UK2Node_FunctionTerminator::ValidateNodeDuringCompilation(FCompilerResultsLog& MessageLog) const
+{
+	Super::ValidateNodeDuringCompilation(MessageLog);
+
+	for (UEdGraphPin* Pin : Pins)
+	{
+		if (Pin && Pin->PinType.bIsWeakPointer && !Pin->PinType.IsContainer())
+		{
+			const FString ErrorString = FString::Printf(
+				*LOCTEXT("WeakPtrNotSupportedError", "Weak prointer is not supported as function parameter. Pin '%s' @@").ToString(), 
+				*Pin->GetName());
+			MessageLog.Error(*ErrorString, this);
+		}
+	}
+}
+
 #undef LOCTEXT_NAMESPACE
diff --git a/Engine/Source/Editor/BlueprintGraph/Private/K2Node_MakeStruct.cpp b/Engine/Source/Editor/BlueprintGraph/Private/K2Node_MakeStruct.cpp
--- a/Engine/Source/Editor/BlueprintGraph/Private/K2Node_MakeStruct.cpp
+++ b/Engine/Source/Editor/BlueprintGraph/Private/K2Node_MakeStruct.cpp
@@ -17,10 +17,19 @@
 
 
 UK2Node_MakeStruct::FMakeStructPinManager::FMakeStructPinManager(const uint8* InSampleStructMemory)
-	: FStructOperationOptionalPinManager(), SampleStructMemory(InSampleStructMemory)
+	: FStructOperationOptionalPinManager()
+	, SampleStructMemory(InSampleStructMemory)
+	, bHasAdvancedPins(false)
 {
 }
 
+void UK2Node_MakeStruct::FMakeStructPinManager::GetRecordDefaults(UProperty* TestProperty, FOptionalPinFromProperty& Record) const
+{
+	UK2Node_StructOperation::FStructOperationOptionalPinManager::GetRecordDefaults(TestProperty, Record);
+	Record.bIsMarkedForAdvancedDisplay = TestProperty ? TestProperty->HasAnyPropertyFlags(CPF_AdvancedDisplay) : false;
+	bHasAdvancedPins |= Record.bIsMarkedForAdvancedDisplay;
+}
+
 void UK2Node_MakeStruct::FMakeStructPinManager::CustomizePinData(UEdGraphPin* Pin, FName SourcePropertyName, int32 ArrayIndex, UProperty* Property) const
 {
 	struct FPinDefaultValueHelper
@@ -67,6 +76,12 @@ void UK2Node_MakeStruct::FMakeStructPinManager::CustomizePinData(UEdGraphPin* Pi
 		checkSlow(bIsObject == ((Schema->PC_Object == Pin->PinType.PinCategory || Schema->PC_Class == Pin->PinType.PinCategory || 
 			Schema->PC_Asset == Pin->PinType.PinCategory || Schema->PC_AssetClass == Pin->PinType.PinCategory) && !Pin->PinType.bIsArray));
 
+		if (Property->HasAnyPropertyFlags(CPF_AdvancedDisplay))
+		{
+			Pin->bAdvancedView = true;
+			bHasAdvancedPins = true;
+		}
+
 		const FString& MetadataDefaultValue = Property->GetMetaData(TEXT("MakeStructureDefaultValue"));
 		if (!MetadataDefaultValue.IsEmpty())
 		{
@@ -130,29 +145,33 @@ void UK2Node_MakeStruct::AllocateDefaultPins()
 	{
 		CreatePin(EGPD_Output, Schema->PC_Struct, TEXT(""), StructType, false, false, StructType->GetName());
 		
+		bool bHasAdvancedPins = false;
 		{
 			FStructOnScope StructOnScope(Cast<UScriptStruct>(StructType));
 			FMakeStructPinManager OptionalPinManager(StructOnScope.GetStructMemory());
 			OptionalPinManager.RebuildPropertyList(ShowPinForProperties, StructType);
 			OptionalPinManager.CreateVisiblePins(ShowPinForProperties, StructType, EGPD_Input, this);
+
+			bHasAdvancedPins = OptionalPinManager.HasAdvancedPins();
 		}
 
 		// When struct has a lot of fields, mark their pins as advanced
-		if(Pins.Num() > 5)
+		if(!bHasAdvancedPins && Pins.Num() > 5)
 		{
-			if(ENodeAdvancedPins::NoPins == AdvancedPinDisplay)
-			{
-				AdvancedPinDisplay = ENodeAdvancedPins::Hidden;
-			}
-
 			for(int32 PinIndex = 3; PinIndex < Pins.Num(); ++PinIndex)
 			{
 				if(UEdGraphPin * EdGraphPin = Pins[PinIndex])
 				{
 					EdGraphPin->bAdvancedView = true;
+					bHasAdvancedPins = true;
 				}
 			}
 		}
+
+		if (bHasAdvancedPins && (ENodeAdvancedPins::NoPins == AdvancedPinDisplay))
+		{
+			AdvancedPinDisplay = ENodeAdvancedPins::Hidden;
+		}
 	}
 }
 
diff --git a/Engine/Source/Editor/BlueprintGraph/Private/K2Node_MathExpression.cpp b/Engine/Source/Editor/BlueprintGraph/Private/K2Node_MathExpression.cpp
--- a/Engine/Source/Editor/BlueprintGraph/Private/K2Node_MathExpression.cpp
+++ b/Engine/Source/Editor/BlueprintGraph/Private/K2Node_MathExpression.cpp
@@ -2688,7 +2688,7 @@ void UK2Node_MathExpression::ValidateNodeDuringCompilation(FCompilerResultsLog&
 			// If the source object is a MacroInstance, we need to look elsewhere for the original MathExpression
 			if(Cast<UK2Node_MacroInstance>(SourceObject))
 			{
-				MathExpression = CastChecked<UK2Node_MathExpression>(MessageLog.FinalNodeBackToMacroSourceMap.FindSourceObject(this));
+				MathExpression = CastChecked<UK2Node_MathExpression>(MessageLog.GetSourceTunnelNode(this));
 			}
 			else
 			{
diff --git a/Engine/Source/Editor/BlueprintGraph/Private/K2Node_SpawnActor.cpp b/Engine/Source/Editor/BlueprintGraph/Private/K2Node_SpawnActor.cpp
--- a/Engine/Source/Editor/BlueprintGraph/Private/K2Node_SpawnActor.cpp
+++ b/Engine/Source/Editor/BlueprintGraph/Private/K2Node_SpawnActor.cpp
@@ -26,7 +26,7 @@ void UK2Node_SpawnActor::AllocateDefaultPins()
 	CreatePin(EGPD_Output, K2Schema->PC_Exec, TEXT(""), NULL, false, false, K2Schema->PN_Then);
 
 	// If required add the world context pin
-	if (GetBlueprint()->ParentClass->HasMetaData(FBlueprintMetadata::MD_ShowWorldContextPin))
+	if (GetBlueprint()->ParentClass->HasMetaDataHierarchical(FBlueprintMetadata::MD_ShowWorldContextPin))
 	{
 		CreatePin(EGPD_Input, K2Schema->PC_Object, TEXT(""), UObject::StaticClass(), false, false, WorldContextPinName);
 	}
diff --git a/Engine/Source/Editor/BlueprintGraph/Private/K2Node_SpawnActorFromClass.cpp b/Engine/Source/Editor/BlueprintGraph/Private/K2Node_SpawnActorFromClass.cpp
--- a/Engine/Source/Editor/BlueprintGraph/Private/K2Node_SpawnActorFromClass.cpp
+++ b/Engine/Source/Editor/BlueprintGraph/Private/K2Node_SpawnActorFromClass.cpp
@@ -42,7 +42,7 @@ void UK2Node_SpawnActorFromClass::AllocateDefaultPins()
 	CreatePin(EGPD_Output, K2Schema->PC_Exec, TEXT(""), NULL, false, false, K2Schema->PN_Then);
 
 	// If required add the world context pin
-	if (GetBlueprint()->ParentClass->HasMetaData(FBlueprintMetadata::MD_ShowWorldContextPin))
+	if (GetBlueprint()->ParentClass->HasMetaDataHierarchical(FBlueprintMetadata::MD_ShowWorldContextPin))
 	{
 		CreatePin(EGPD_Input, K2Schema->PC_Object, TEXT(""), UObject::StaticClass(), false, false, FK2Node_SpawnActorFromClassHelper::WorldContextPinName);
 	}
diff --git a/Engine/Source/Editor/BlueprintGraph/Private/K2Node_TunnelBoundary.cpp b/Engine/Source/Editor/BlueprintGraph/Private/K2Node_TunnelBoundary.cpp
--- a/Engine/Source/Editor/BlueprintGraph/Private/K2Node_TunnelBoundary.cpp
+++ b/Engine/Source/Editor/BlueprintGraph/Private/K2Node_TunnelBoundary.cpp
@@ -25,7 +25,7 @@ public:
 				case TBT_EntrySite:
 				case TBT_ExitSite:
 				{
-		GenerateSimpleThenGoto(Context, *Node);
+					GenerateSimpleThenGoto(Context, *Node);
 					break;
 				}
 				case TBT_EndOfThread:
@@ -55,7 +55,7 @@ UK2Node_TunnelBoundary::UK2Node_TunnelBoundary(const FObjectInitializer& ObjectI
 FText UK2Node_TunnelBoundary::GetNodeTitle(ENodeTitleType::Type TitleType) const
 {
 	FFormatNamedArguments Args;
-	Args.Add(TEXT("GraphName"), FText::FromName(TunnelGraphName));
+	Args.Add(TEXT("BaseName"), FText::FromName(BaseName));
 	switch(TunnelBoundaryType)
 	{
 		case TBT_EntrySite:
@@ -74,7 +74,7 @@ FText UK2Node_TunnelBoundary::GetNodeTitle(ENodeTitleType::Type TitleType) const
 			break;
 		}
 	}
-	return FText::Format(LOCTEXT("UK2Node_TunnelBoundary_FullTitle", "{GraphName} {EntryType}"), Args);
+	return FText::Format(LOCTEXT("UK2Node_TunnelBoundary_FullTitle", "{BaseName} {EntryType}"), Args);
 }
 
 FNodeHandlingFunctor* UK2Node_TunnelBoundary::CreateNodeHandler(FKismetCompilerContext& CompilerContext) const
@@ -82,7 +82,7 @@ FNodeHandlingFunctor* UK2Node_TunnelBoundary::CreateNodeHandler(FKismetCompilerC
 	return new FKCHandler_TunnelBoundary(CompilerContext);
 }
 
-void UK2Node_TunnelBoundary::CreateBoundaryNodesForGraph(UEdGraph* TunnelGraph, class FCompilerResultsLog& MessageLog)
+void UK2Node_TunnelBoundary::CreateBoundaryNodesForGraph(UEdGraph* TunnelGraph, FCompilerResultsLog& MessageLog)
 {
 	if (UK2Node_Tunnel* OuterTunnel = TunnelGraph->GetTypedOuter<UK2Node_Tunnel>())
 	{
@@ -91,7 +91,7 @@ void UK2Node_TunnelBoundary::CreateBoundaryNodesForGraph(UEdGraph* TunnelGraph,
 	}
 }
 
-void UK2Node_TunnelBoundary::CreateBoundaryNodesForTunnelInstance(UK2Node_Tunnel* TunnelInstance, UEdGraph* TunnelGraph, class FCompilerResultsLog& MessageLog)
+void UK2Node_TunnelBoundary::CreateBoundaryNodesForTunnelInstance(UK2Node_Tunnel* TunnelInstance, UEdGraph* TunnelGraph, FCompilerResultsLog& MessageLog)
 {
 	if (!TunnelGraph)
 	{
@@ -127,6 +127,7 @@ void UK2Node_TunnelBoundary::CreateBoundaryNodesForTunnelInstance(UK2Node_Tunnel
 	// Find the blueprint tunnel instance.
 	UEdGraphNode* SourceTunnelInstance = Cast<UEdGraphNode>(MessageLog.FindSourceObject(TunnelInstance));
 	SourceTunnelInstance = FindTrueSourceTunnelInstance(TunnelInstance, SourceTunnelInstance);
+	UEdGraphNode* RegistrantTunnelInstance = TunnelInstance->IsA<UK2Node_Composite>() ? SourceTunnelInstance : TunnelInstance;
 	// Create the Boundary nodes for each unique entry/exit site.
 	TArray<UEdGraphPin*> ExecutionEndpointPins;
 	TArray<UEdGraphPin*> VisitedPins;
@@ -135,7 +136,7 @@ void UK2Node_TunnelBoundary::CreateBoundaryNodesForTunnelInstance(UK2Node_Tunnel
 		UK2Node_TunnelBoundary* TunnelBoundaryNode = TunnelGraph->CreateBlankNode<UK2Node_TunnelBoundary>();
 		MessageLog.NotifyIntermediateObjectCreation(TunnelBoundaryNode, TunnelInstance);
 		TunnelBoundaryNode->CreateNewGuid();
-		MessageLog.IntermediateTunnelNodeToSourceNodeMap.Add(TunnelBoundaryNode, SourceTunnelInstance);
+		MessageLog.RegisterIntermediateTunnelNode(TunnelBoundaryNode, RegistrantTunnelInstance);
 		TunnelBoundaryNode->WireUpTunnelEntry(TunnelInstance, EntryPin, MessageLog);
 		for (auto LinkedPin : EntryPin->LinkedTo)
 		{
@@ -149,7 +150,7 @@ void UK2Node_TunnelBoundary::CreateBoundaryNodesForTunnelInstance(UK2Node_Tunnel
 		MessageLog.NotifyIntermediateObjectCreation(TunnelBoundaryNode, TunnelInstance);
 		TunnelBoundaryNode->CreateNewGuid();
 		TunnelBoundaryNode->TunnelBoundaryType = TBT_EndOfThread;
-		MessageLog.IntermediateTunnelNodeToSourceNodeMap.Add(TunnelBoundaryNode, SourceTunnelInstance);
+		MessageLog.RegisterIntermediateTunnelNode(TunnelBoundaryNode, RegistrantTunnelInstance);
 		// Create exit point pin
 		UEdGraphPin* BoundaryTerminalPin = TunnelBoundaryNode->CreatePin(EGPD_Input, UEdGraphSchema_K2::PC_Exec, TEXT(""), NULL, false, false, TEXT("TunnelEndOfThread"));
 		// Wire up the endpoint unwired links to the boundary node.
@@ -163,25 +164,65 @@ void UK2Node_TunnelBoundary::CreateBoundaryNodesForTunnelInstance(UK2Node_Tunnel
 		UK2Node_TunnelBoundary* TunnelBoundaryNode = TunnelGraph->CreateBlankNode<UK2Node_TunnelBoundary>();
 		MessageLog.NotifyIntermediateObjectCreation(TunnelBoundaryNode, TunnelInstance);
 		TunnelBoundaryNode->CreateNewGuid();
-		MessageLog.IntermediateTunnelNodeToSourceNodeMap.Add(TunnelBoundaryNode, SourceTunnelInstance);
+		MessageLog.RegisterIntermediateTunnelNode(TunnelBoundaryNode, RegistrantTunnelInstance);
 		TunnelBoundaryNode->WireUpTunnelExit(TunnelInstance, ExitPin, MessageLog);
 	}
 	// Build node guid map to locate true source node.
 	TMap<FGuid, const UEdGraphNode*> TrueSourceNodeMap;
 	BuildSourceNodeMap(SourceTunnelInstance, TrueSourceNodeMap);
-	// Register all the nodes to the tunnel instance map
-	if (SourceTunnelInstance)
+}
+
+void UK2Node_TunnelBoundary::CreateBoundariesForExpansionNodes(UEdGraphNode* SourceNode, TArray<UEdGraphNode*>& ExpansionNodes, TMap<UEdGraphPin*, UEdGraphPin*>& LinkedPinMap, FCompilerResultsLog& MessageLog)
+{
+	// Find boundary points and map them to the source pins
+	UEdGraph* TargetGraph = nullptr;
+	TMap<UEdGraphPin*, TArray<UEdGraphPin*>> EntryPins;
+	TMap<UEdGraphPin*, TArray<UEdGraphPin*>> ExitPins;
+	for (auto ExpansionNode : ExpansionNodes)
 	{
-		for (auto Node : TunnelGraph->Nodes)
+		if (!TargetGraph)
 		{
-			MessageLog.SourceNodeToTunnelInstanceNodeMap.Add(Node, SourceTunnelInstance);
-			if (const UEdGraphNode** SearchNode = TrueSourceNodeMap.Find(Node->NodeGuid))
+			TargetGraph = ExpansionNode->GetGraph();
+		}
+		for (UEdGraphPin* Pin : ExpansionNode->Pins)
+		{
+			for (UEdGraphPin* LinkedPin : Pin->LinkedTo)
 			{
-				const UEdGraphNode* SourceGraphNode = *SearchNode;
-				MessageLog.IntermediateTunnelNodeToSourceNodeMap.Add(Node, SourceGraphNode);
+				if (UEdGraphPin** SourcePin = LinkedPinMap.Find(LinkedPin))
+				{
+					if (Pin->Direction == EGPD_Input)
+					{
+						TArray<UEdGraphPin*>& PinLinks = EntryPins.FindOrAdd(*SourcePin);
+						PinLinks.Add(Pin);
+					}
+					else
+					{
+						TArray<UEdGraphPin*>& PinLinks = ExitPins.FindOrAdd(*SourcePin);
+						PinLinks.Add(Pin);
+					}
+				}
 			}
 		}
-		MessageLog.SourceNodeToTunnelInstanceNodeMap.Add(TunnelInstance, SourceTunnelInstance);
+	}
+	// Create the Boundary nodes for each unique entry/exit site.
+	if (TargetGraph)
+	{
+		// Create entry node and pins
+		UK2Node_TunnelBoundary* EntryBoundaryNode = TargetGraph->CreateBlankNode<UK2Node_TunnelBoundary>();
+		MessageLog.NotifyIntermediateObjectCreation(EntryBoundaryNode, SourceNode);
+		EntryBoundaryNode->CreateNewGuid();
+		for (auto PinEntry : EntryPins)
+		{
+			EntryBoundaryNode->WireUpEntry(SourceNode, PinEntry.Key, PinEntry.Value, MessageLog);
+		}
+		// Create exit node and pins
+		UK2Node_TunnelBoundary* ExitBoundaryNode = TargetGraph->CreateBlankNode<UK2Node_TunnelBoundary>();
+		MessageLog.NotifyIntermediateObjectCreation(ExitBoundaryNode, SourceNode);
+		ExitBoundaryNode->CreateNewGuid();
+		for (auto PinExit : ExitPins)
+		{
+			ExitBoundaryNode->WireUpExit(SourceNode, PinExit.Key, PinExit.Value, MessageLog);
+		}
 	}
 }
 
@@ -197,7 +238,7 @@ void UK2Node_TunnelBoundary::WireUpTunnelEntry(UK2Node_Tunnel* TunnelInstance, U
 		// Mark as entry node
 		TunnelBoundaryType = TBT_EntrySite;
 		// Grab the graph name
-		DetermineTunnelGraphName(TunnelInstance);
+		CreateBaseNodeName(TunnelInstance);
 		// Find the true source pin
 		UEdGraphPin* SourcePin = TunnelInstance->FindPin(TunnelPin->PinName);
 		// Wire in the node
@@ -220,7 +261,7 @@ void UK2Node_TunnelBoundary::WireUpTunnelExit(UK2Node_Tunnel* TunnelInstance, UE
 		// Mark as exit node
 		TunnelBoundaryType = TBT_ExitSite;
 		// Grab the graph name
-		DetermineTunnelGraphName(TunnelInstance);
+		CreateBaseNodeName(TunnelInstance);
 		// Find the true source pin
 		UEdGraphPin* SourcePin = TunnelInstance->FindPin(TunnelPin->PinName);
 		// Wire in the node
@@ -236,18 +277,77 @@ void UK2Node_TunnelBoundary::WireUpTunnelExit(UK2Node_Tunnel* TunnelInstance, UE
 	}
 }
 
-void UK2Node_TunnelBoundary::DetermineTunnelGraphName(UK2Node_Tunnel* TunnelInstance)
+void UK2Node_TunnelBoundary::WireUpEntry(UEdGraphNode* SourceNode, UEdGraphPin* SourcePin, TArray<UEdGraphPin*>& EntryPins, FCompilerResultsLog& MessageLog)
+{
+	if (SourceNode && SourcePin && EntryPins.Num())
+	{
+		// Mark as entry node
+		TunnelBoundaryType = TBT_EntrySite;
+		// Create the pin base name
+		CreateBaseNodeName(SourceNode);
+		// Wire in the node
+		UEdGraphPin* InputPin = CreatePin(EGPD_Input, SourcePin->PinType.PinCategory, SourcePin->PinType.PinSubCategory, SourcePin->PinType.PinSubCategoryObject.Get(), SourcePin->PinType.bIsArray, SourcePin->PinType.bIsReference, TEXT("EntryBoundary"), SourcePin->PinType.bIsConst);
+		MessageLog.NotifyIntermediatePinCreation(InputPin, SourcePin);
+		UEdGraphPin* OutputPin = CreatePin(EGPD_Output, SourcePin->PinType.PinCategory, SourcePin->PinType.PinSubCategory, SourcePin->PinType.PinSubCategoryObject.Get(), SourcePin->PinType.bIsArray, SourcePin->PinType.bIsReference, SourcePin->PinName, SourcePin->PinType.bIsConst);
+		MessageLog.AddPinTraceFilter(OutputPin);
+		for (auto EntryPin : EntryPins)
+		{
+			for (auto LinkedPin : EntryPin->LinkedTo)
+			{
+				check (LinkedPin->Direction != InputPin->Direction);
+				LinkedPin->MakeLinkTo(InputPin);
+			}
+			EntryPin->BreakAllPinLinks();
+			check (EntryPin->Direction != OutputPin->Direction);
+			EntryPin->MakeLinkTo(OutputPin);
+		}
+	}
+}
+
+void UK2Node_TunnelBoundary::WireUpExit(UEdGraphNode* SourceNode, UEdGraphPin* SourcePin, TArray<UEdGraphPin*>& ExitPins, FCompilerResultsLog& MessageLog)
+{
+	if (SourceNode && SourcePin && ExitPins.Num())
+	{
+		// Mark as exit node
+		TunnelBoundaryType = TBT_ExitSite;
+		// Create the pin base name
+		CreateBaseNodeName(SourceNode);
+		// Wire in the node
+		UEdGraphPin* OutputPin = CreatePin(EGPD_Output, SourcePin->PinType.PinCategory, SourcePin->PinType.PinSubCategory, SourcePin->PinType.PinSubCategoryObject.Get(), SourcePin->PinType.bIsArray, SourcePin->PinType.bIsReference, SourcePin->PinName, SourcePin->PinType.bIsConst);
+		MessageLog.NotifyIntermediatePinCreation(OutputPin, SourcePin);
+		UEdGraphPin* InputPin = CreatePin(EGPD_Input, SourcePin->PinType.PinCategory, SourcePin->PinType.PinSubCategory, SourcePin->PinType.PinSubCategoryObject.Get(), SourcePin->PinType.bIsArray, SourcePin->PinType.bIsReference, TEXT("ExitBoundary"), SourcePin->PinType.bIsConst);
+		MessageLog.AddPinTraceFilter(InputPin);
+		for (auto ExitPin : ExitPins)
+		{
+			for (auto LinkedPin : ExitPin->LinkedTo)
+			{
+				check (LinkedPin->Direction != OutputPin->Direction);
+				LinkedPin->MakeLinkTo(OutputPin);
+			}
+			ExitPin->BreakAllPinLinks();
+			check (ExitPin->Direction != InputPin->Direction);
+			ExitPin->MakeLinkTo(InputPin);
+		}
+	}
+}
+
+void UK2Node_TunnelBoundary::CreateBaseNodeName(UEdGraphNode* SourceNode)
 {
 	UEdGraph* TunnelGraph = nullptr;
-	if (const UK2Node_MacroInstance* MacroInstance = Cast<UK2Node_MacroInstance>(TunnelInstance))
+	if (const UK2Node_MacroInstance* MacroInstance = Cast<UK2Node_MacroInstance>(SourceNode))
 	{
 		TunnelGraph = MacroInstance->GetMacroGraph();
+		BaseName = TunnelGraph ? TunnelGraph->GetFName() : NAME_None;
 	}
-	else if (const UK2Node_Composite* CompositeInstance = Cast<UK2Node_Composite>(TunnelInstance))
+	else if (const UK2Node_Composite* CompositeInstance = Cast<UK2Node_Composite>(SourceNode))
 	{
 		TunnelGraph = CompositeInstance->BoundGraph;
+		BaseName = TunnelGraph ? TunnelGraph->GetFName() : NAME_None;
+	}
+	else
+	{
+		BaseName = SourceNode ? SourceNode->GetFName() : NAME_None;
 	}
-	TunnelGraphName = TunnelGraph ? TunnelGraph->GetFName() : NAME_None;
 }
 
 void UK2Node_TunnelBoundary::BuildSourceNodeMap(UEdGraphNode* Tunnel, TMap<FGuid, const UEdGraphNode*>& SourceNodeMap)
diff --git a/Engine/Source/Editor/DetailCustomizations/Private/SceneComponentDetails.cpp b/Engine/Source/Editor/DetailCustomizations/Private/SceneComponentDetails.cpp
--- a/Engine/Source/Editor/DetailCustomizations/Private/SceneComponentDetails.cpp
+++ b/Engine/Source/Editor/DetailCustomizations/Private/SceneComponentDetails.cpp
@@ -62,7 +62,7 @@ static EComponentMobility::Type GetInheritedMobility(USceneComponent const* cons
 				for (USCS_Node const* Node : BlueprintClass->SimpleConstructionScript->GetAllNodes())
 				{
 					check(Node); // fix for this bug: https://connect.microsoft.com/VisualStudio/feedback/details/3081898
-					if (Node->VariableName == ComponentNode->ParentComponentOrVariableName)
+					if (Node->GetVariableName() == ComponentNode->ParentComponentOrVariableName)
 					{
 						ParentNode = Node;
 						break;
diff --git a/Engine/Source/Editor/FoliageEdit/Private/FoliageEdMode.cpp b/Engine/Source/Editor/FoliageEdit/Private/FoliageEdMode.cpp
--- a/Engine/Source/Editor/FoliageEdit/Private/FoliageEdMode.cpp
+++ b/Engine/Source/Editor/FoliageEdit/Private/FoliageEdMode.cpp
@@ -452,13 +452,6 @@ void FEdModeFoliage::OnVRHoverUpdate(FEditorViewportClient& ViewportClient, UVie
 				{
 					EndFoliageBrushTrace();
 				}
-				else
-				{
-					if (UISettings.GetPaintToolSelected() || UISettings.GetReapplyToolSelected() || UISettings.GetLassoSelectToolSelected())
-					{
-						StartFoliageBrushTrace(&ViewportClient, ViewportInteractor);
-					}
-				}
 
 				// Skip other interactors if we are painting with one
 				if (ViewportInteractor && ViewportInteractor == FoliageInteractor)
@@ -760,11 +753,6 @@ void FEdModeFoliage::Tick(FEditorViewportClient* ViewportClient, float DeltaTime
 					}
 					else if (UISettings.GetPaintToolSelected() || UISettings.GetReapplyToolSelected() || UISettings.GetLassoSelectToolSelected())
 					{
-						if (UISettings.GetPaintToolSelected())
-						{
-							StartFoliageBrushTrace(ViewportClient, Interactor);
-						}
-
 						// Go ahead and paint immediately
 						FVector LaserPointerStart, LaserPointerEnd;
 						if (Interactor->GetLaserPointer( /* Out */ LaserPointerStart, /* Out */ LaserPointerEnd))
diff --git a/Engine/Source/Editor/GameplayTagsEditor/Classes/GameplayTagsK2Node_LiteralGameplayTag.h b/Engine/Source/Editor/GameplayTagsEditor/Classes/GameplayTagsK2Node_LiteralGameplayTag.h
--- a/Engine/Source/Editor/GameplayTagsEditor/Classes/GameplayTagsK2Node_LiteralGameplayTag.h
+++ b/Engine/Source/Editor/GameplayTagsEditor/Classes/GameplayTagsK2Node_LiteralGameplayTag.h
@@ -26,6 +26,7 @@ class UGameplayTagsK2Node_LiteralGameplayTag : public UK2Node
 	virtual void ExpandNode(class FKismetCompilerContext& CompilerContext, UEdGraph* SourceGraph) override;
 	virtual void GetMenuActions(FBlueprintActionDatabaseRegistrar& ActionRegistrar) const override;
 	virtual FText GetMenuCategory() const override;
+	virtual bool IsNodePure() const override { return true; }
 	//~ End UK2Node Interface
 #endif
 
diff --git a/Engine/Source/Editor/GameplayTagsEditor/Classes/GameplayTagsK2Node_MultiCompareBase.h b/Engine/Source/Editor/GameplayTagsEditor/Classes/GameplayTagsK2Node_MultiCompareBase.h
--- a/Engine/Source/Editor/GameplayTagsEditor/Classes/GameplayTagsK2Node_MultiCompareBase.h
+++ b/Engine/Source/Editor/GameplayTagsEditor/Classes/GameplayTagsK2Node_MultiCompareBase.h
@@ -32,6 +32,7 @@ class UGameplayTagsK2Node_MultiCompareBase : public UK2Node
 	virtual bool NodeCausesStructuralBlueprintChange() const override { return true; }
 	virtual bool ShouldShowNodeProperties() const override { return true; }
 	virtual FText GetMenuCategory() const override;
+	virtual bool IsNodePure() const override { return true; }
 	// End of UK2Node interface
 
 	void AddPin();
diff --git a/Engine/Source/Editor/GraphEditor/Private/SGraphPin.cpp b/Engine/Source/Editor/GraphEditor/Private/SGraphPin.cpp
--- a/Engine/Source/Editor/GraphEditor/Private/SGraphPin.cpp
+++ b/Engine/Source/Editor/GraphEditor/Private/SGraphPin.cpp
@@ -863,18 +863,21 @@ const FSlateBrush* SGraphPin::GetPinBorder() const
 
 FSlateColor SGraphPin::GetPinColor() const
 {
-	if(GraphPinObj->bIsDiffing)
+	if (!GraphPinObj->IsPendingKill())
 	{
-		return FSlateColor(FLinearColor(0.9f,0.2f,0.15f));
-	}
-	if (const UEdGraphSchema* Schema = GraphPinObj->GetSchema())
-	{
-		if(!GetPinObj()->GetOwningNode()->IsNodeEnabled() || !IsEditingEnabled())
+		if (GraphPinObj->bIsDiffing)
 		{
-			return Schema->GetPinTypeColor(GraphPinObj->PinType) * FLinearColor(1.0f, 1.0f, 1.0f, 0.5f);
+			return FSlateColor(FLinearColor(0.9f, 0.2f, 0.15f));
 		}
+		if (const UEdGraphSchema* Schema = GraphPinObj->GetSchema())
+		{
+			if (!GetPinObj()->GetOwningNode()->IsNodeEnabled() || !IsEditingEnabled())
+			{
+				return Schema->GetPinTypeColor(GraphPinObj->PinType) * FLinearColor(1.0f, 1.0f, 1.0f, 0.5f);
+			}
 
-		return Schema->GetPinTypeColor(GraphPinObj->PinType) * PinColorModifier;
+			return Schema->GetPinTypeColor(GraphPinObj->PinType) * PinColorModifier;
+		}
 	}
 
 	return FLinearColor::White;
@@ -900,15 +903,18 @@ FSlateColor SGraphPin::GetPinTextColor() const
 
 const FSlateBrush* SGraphPin::GetPinStatusIcon() const
 {
-	UEdGraphPin* WatchedPin = ((GraphPinObj->Direction == EGPD_Input) && (GraphPinObj->LinkedTo.Num() > 0)) ? GraphPinObj->LinkedTo[0] : GraphPinObj;
-
-	if (UEdGraphNode* GraphNode = WatchedPin->GetOwningNodeUnchecked())
+	if (!GraphPinObj->IsPendingKill())
 	{
-		UBlueprint* Blueprint = FBlueprintEditorUtils::FindBlueprintForNodeChecked(GraphNode);
+		UEdGraphPin* WatchedPin = ((GraphPinObj->Direction == EGPD_Input) && (GraphPinObj->LinkedTo.Num() > 0)) ? GraphPinObj->LinkedTo[0] : GraphPinObj;
 
-		if (FKismetDebugUtilities::IsPinBeingWatched(Blueprint, WatchedPin) )
+		if (UEdGraphNode* GraphNode = WatchedPin->GetOwningNodeUnchecked())
 		{
-			return FEditorStyle::GetBrush( TEXT("Graph.WatchedPinIcon_Pinned") );
+			UBlueprint* Blueprint = FBlueprintEditorUtils::FindBlueprintForNodeChecked(GraphNode);
+
+			if (FKismetDebugUtilities::IsPinBeingWatched(Blueprint, WatchedPin))
+			{
+				return FEditorStyle::GetBrush(TEXT("Graph.WatchedPinIcon_Pinned"));
+			}
 		}
 	}
 
@@ -917,6 +923,11 @@ const FSlateBrush* SGraphPin::GetPinStatusIcon() const
 
 EVisibility SGraphPin::GetPinStatusIconVisibility() const
 {
+	if (GraphPinObj->IsPendingKill())
+	{
+		return EVisibility::Collapsed;
+	}
+
 	UEdGraphPin const* WatchedPin = ((GraphPinObj->Direction == EGPD_Input) && (GraphPinObj->LinkedTo.Num() > 0)) ? GraphPinObj->LinkedTo[0] : GraphPinObj;
 
 	UEdGraphSchema const* Schema = GraphPinObj->GetSchema();
@@ -925,6 +936,11 @@ EVisibility SGraphPin::GetPinStatusIconVisibility() const
 
 FReply SGraphPin::ClickedOnPinStatusIcon()
 {
+	if (GraphPinObj->IsPendingKill())
+	{
+		return FReply::Handled();
+	}
+
 	UEdGraphPin* WatchedPin = ((GraphPinObj->Direction == EGPD_Input) && (GraphPinObj->LinkedTo.Num() > 0)) ? GraphPinObj->LinkedTo[0] : GraphPinObj;
 
 	UEdGraphSchema const* Schema = GraphPinObj->GetSchema();
@@ -936,7 +952,7 @@ FReply SGraphPin::ClickedOnPinStatusIcon()
 EVisibility SGraphPin::GetDefaultValueVisibility() const
 {
 	// First ask schema
-	const UEdGraphSchema* Schema = GraphPinObj->GetSchema();
+	const UEdGraphSchema* Schema = !GraphPinObj->IsPendingKill() ? GraphPinObj->GetSchema() : nullptr;
 	if (Schema == nullptr || Schema->ShouldHidePinDefaultValue(GraphPinObj))
 	{
 		return EVisibility::Collapsed;
diff --git a/Engine/Source/Editor/Kismet/Private/BlueprintDetailsCustomization.cpp b/Engine/Source/Editor/Kismet/Private/BlueprintDetailsCustomization.cpp
--- a/Engine/Source/Editor/Kismet/Private/BlueprintDetailsCustomization.cpp
+++ b/Engine/Source/Editor/Kismet/Private/BlueprintDetailsCustomization.cpp
@@ -1015,7 +1015,7 @@ bool FBlueprintVarActionDetails::GetVariableNameChangeEnabled() const
 		{
 			if (USCS_Node* Node = BlueprintObj->SimpleConstructionScript->FindSCSNode(CachedVariableName))
 			{
-				bIsReadOnly = !FComponentEditorUtils::IsValidVariableNameString(Node->ComponentTemplate, Node->VariableName.ToString());
+				bIsReadOnly = !FComponentEditorUtils::IsValidVariableNameString(Node->ComponentTemplate, Node->GetVariableName().ToString());
 			}
 		}
 		else if(IsALocalVariable(VariableProperty))
@@ -1044,7 +1044,7 @@ void FBlueprintVarActionDetails::OnVarNameChanged(const FText& InNewText)
 	{
 		for (USCS_Node* Node : BlueprintObj->SimpleConstructionScript->GetAllNodes())
 		{
-			if (Node && Node->VariableName == CachedVariableName && !FComponentEditorUtils::IsValidVariableNameString(Node->ComponentTemplate, InNewText.ToString()))
+			if (Node && Node->GetVariableName() == CachedVariableName && !FComponentEditorUtils::IsValidVariableNameString(Node->ComponentTemplate, InNewText.ToString()))
 			{
 				VarNameEditableTextBox->SetError(LOCTEXT("ComponentVariableRenameFailed_NotValid", "This name is reserved for engine use."));
 				return;
diff --git a/Engine/Source/Editor/Kismet/Private/BlueprintEditor.cpp b/Engine/Source/Editor/Kismet/Private/BlueprintEditor.cpp
--- a/Engine/Source/Editor/Kismet/Private/BlueprintEditor.cpp
+++ b/Engine/Source/Editor/Kismet/Private/BlueprintEditor.cpp
@@ -1459,7 +1459,7 @@ void FBlueprintEditor::EnsureBlueprintIsUpToDate(UBlueprint* BlueprintObj)
 			{
 				if(SCS_Node)
 				{
-					FBlueprintEditorUtils::RemoveVariableNodes(BlueprintObj, SCS_Node->VariableName);
+					FBlueprintEditorUtils::RemoveVariableNodes(BlueprintObj, SCS_Node->GetVariableName());
 				}
 			}
 		
@@ -3105,7 +3105,7 @@ void FBlueprintEditor::OnBlueprintChangedImpl(UBlueprint* InBlueprint, bool bIsJ
 		// Notify that the blueprint has been changed (update Content browser, etc)
 		InBlueprint->PostEditChange();
 
-		// Call PostEditChange() on any Actors that are based on this Blueprint
+		// Call PostEditChange() on any Actors that are based on this Blueprint 
 		FBlueprintEditorUtils::PostEditChangeBlueprintActors(InBlueprint);
 
 		// Refresh the graphs
@@ -6494,7 +6494,7 @@ void FBlueprintEditor::ExtractEventTemplateForFunction(class UK2Node_CustomEvent
 	}
 }
 
-void FBlueprintEditor::CollapseNodesIntoGraph(UEdGraphNode* InGatewayNode, UK2Node_EditablePinBase* InEntryNode, UK2Node_EditablePinBase* InResultNode, UEdGraph* InSourceGraph, UEdGraph* InDestinationGraph, TSet<UEdGraphNode*>& InCollapsableNodes, bool bCanDiscardEmptyReturnNode)
+void FBlueprintEditor::CollapseNodesIntoGraph(UEdGraphNode* InGatewayNode, UK2Node_EditablePinBase* InEntryNode, UK2Node_EditablePinBase* InResultNode, UEdGraph* InSourceGraph, UEdGraph* InDestinationGraph, TSet<UEdGraphNode*>& InCollapsableNodes, bool bCanDiscardEmptyReturnNode, bool bCanHaveWeakObjPtrParam)
 {
 	const UEdGraphSchema_K2* K2Schema = GetDefault<UEdGraphSchema_K2>();
 
@@ -6633,8 +6633,13 @@ void FBlueprintEditor::CollapseNodesIntoGraph(UEdGraphNode* InGatewayNode, UK2No
 								bDiscardReturnNode = false;
 							}
 							
-							RemotePortPin = InGatewayNode->CreatePin(LocalPin->Direction, LocalPin->PinType, UniquePortName);
-							LocalPortPin = LocalPort->CreateUserDefinedPin(UniquePortName, LocalPin->PinType, (LocalPin->Direction == EGPD_Input)? EGPD_Output : EGPD_Input);
+							FEdGraphPinType PinType = LocalPin->PinType;
+							if (PinType.bIsWeakPointer && !PinType.IsContainer() && !bCanHaveWeakObjPtrParam)
+							{
+								PinType.bIsWeakPointer = false;
+							}
+							RemotePortPin = InGatewayNode->CreatePin(LocalPin->Direction, PinType, UniquePortName);
+							LocalPortPin = LocalPort->CreateUserDefinedPin(UniquePortName, PinType, (LocalPin->Direction == EGPD_Input)? EGPD_Output : EGPD_Input);
 						}
 					}
 				}
@@ -6757,7 +6762,7 @@ void FBlueprintEditor::CollapseNodes(TSet<UEdGraphNode*>& InCollapsableNodes)
 	TSharedPtr<INameValidatorInterface> NameValidator = MakeShareable(new FKismetNameValidator(GetBlueprintObj(), GraphName));
 	FBlueprintEditorUtils::RenameGraphWithSuggestion(DestinationGraph, NameValidator, GraphName.ToString());
 
-	CollapseNodesIntoGraph(GatewayNode, GatewayNode->GetInputSink(), GatewayNode->GetOutputSource(), SourceGraph, DestinationGraph, InCollapsableNodes);
+	CollapseNodesIntoGraph(GatewayNode, GatewayNode->GetInputSink(), GatewayNode->GetOutputSource(), SourceGraph, DestinationGraph, InCollapsableNodes, false, true);
 }
 
 UEdGraph* FBlueprintEditor::CollapseSelectionToFunction(TSharedPtr<SGraphEditor> InRootGraph, TSet<class UEdGraphNode*>& InCollapsableNodes, UEdGraphNode*& OutFunctionNode)
@@ -6807,7 +6812,7 @@ UEdGraph* FBlueprintEditor::CollapseSelectionToFunction(TSharedPtr<SGraphEditor>
 
 	check(OutFunctionNode);
 
-	CollapseNodesIntoGraph(OutFunctionNode, EntryNode, ResultNode, SourceGraph, NewGraph, InCollapsableNodes, true);
+	CollapseNodesIntoGraph(OutFunctionNode, EntryNode, ResultNode, SourceGraph, NewGraph, InCollapsableNodes, true, false);
 
 	FBlueprintEditorUtils::MarkBlueprintAsStructurallyModified(GetBlueprintObj());
 	OutFunctionNode->ReconstructNode();
@@ -6860,7 +6865,7 @@ UEdGraph* FBlueprintEditor::CollapseSelectionToMacro(TSharedPtr<SGraphEditor> In
 		}
 	}
 
-	CollapseNodesIntoGraph(GatewayNode, InputSink, OutputSink, SourceGraph, DestinationGraph, InCollapsableNodes);
+	CollapseNodesIntoGraph(GatewayNode, InputSink, OutputSink, SourceGraph, DestinationGraph, InCollapsableNodes, false, false);
 
 	OutMacroNode = GatewayNode;
 	OutMacroNode->ReconstructNode();
@@ -8048,24 +8053,9 @@ void FBlueprintEditor::UpdatePreviewActor(UBlueprint* InBlueprint, bool bInForce
 			SpawnInfo.bNoFail = true;
 			SpawnInfo.ObjectFlags = RF_Transient|RF_Transactional;
 
-			// Temporarily remove the deprecated and abstract flag so we can respawn the Blueprint in the viewport
-			const bool bIsClassDeprecated = PreviewBlueprint->GeneratedClass->HasAnyClassFlags(CLASS_Deprecated);
-			const bool bIsClassAbstract = PreviewBlueprint->GeneratedClass->HasAnyClassFlags(CLASS_Abstract);
-			PreviewBlueprint->GeneratedClass->ClassFlags &= ~CLASS_Deprecated;
-			PreviewBlueprint->GeneratedClass->ClassFlags &= ~CLASS_Abstract;
-
-			PreviewActorPtr = PreviewActor = PreviewScene.GetWorld()->SpawnActor(PreviewBlueprint->GeneratedClass, &SpawnLocation, &SpawnRotation, SpawnInfo);
-
-			// Reassign the deprecated flag if it was previously assigned
-			if ( bIsClassDeprecated )
-			{
-				PreviewBlueprint->GeneratedClass->ClassFlags |= CLASS_Deprecated;
-			}
-			 
-			// Reassign the abstract flag if it was previously assigned
-			if (bIsClassAbstract)
 			{
-				PreviewBlueprint->GeneratedClass->ClassFlags |= CLASS_Abstract;
+				FMakeClassSpawnableOnScope TemporarilySpawnable(PreviewBlueprint->GeneratedClass);
+				PreviewActorPtr = PreviewActor = PreviewScene.GetWorld()->SpawnActor(PreviewBlueprint->GeneratedClass, &SpawnLocation, &SpawnRotation, SpawnInfo);
 			}
 
 			check(PreviewActor);
diff --git a/Engine/Source/Editor/Kismet/Private/FindInBlueprints.cpp b/Engine/Source/Editor/Kismet/Private/FindInBlueprints.cpp
--- a/Engine/Source/Editor/Kismet/Private/FindInBlueprints.cpp
+++ b/Engine/Source/Editor/Kismet/Private/FindInBlueprints.cpp
@@ -412,7 +412,7 @@ FReply FFindInBlueprintsProperty::OnClick()
 				const TArray<USCS_Node*>& Nodes = Blueprint->SimpleConstructionScript->GetAllNodes();
 				for (USCS_Node* Node : Nodes)
 				{
-					if (Node->VariableName.ToString() == DisplayText.ToString())
+					if (Node->GetVariableName().ToString() == DisplayText.ToString())
 					{
 						UBlueprintGeneratedClass* GeneratedClass = Cast<UBlueprintGeneratedClass>(Blueprint->GeneratedClass);
 						if (GeneratedClass)
diff --git a/Engine/Source/Editor/Kismet/Private/SBlueprintDiff.cpp b/Engine/Source/Editor/Kismet/Private/SBlueprintDiff.cpp
--- a/Engine/Source/Editor/Kismet/Private/SBlueprintDiff.cpp
+++ b/Engine/Source/Editor/Kismet/Private/SBlueprintDiff.cpp
@@ -181,69 +181,73 @@ FCDODiffControl::FCDODiffControl(
 	TArray<FPropertySoftPath> OldProperties = OldDetails.GetDisplayedProperties();
 	TArray<FPropertySoftPath> NewProperties = NewDetails.GetDisplayedProperties();
 
-	const auto FindDiffering = [](TArray< FSingleObjectDiffEntry > const& InDifferenceEntries, const FPropertySoftPath& PropertyIdentifer) -> const FSingleObjectDiffEntry*
+	const auto FindAndPushDiff = [&OrderedProperties, &DifferingProperties](const FPropertySoftPath& PropertyIdentifer) -> bool
 	{
-		for (const auto& Difference : InDifferenceEntries)
+		for (const auto& Difference : DifferingProperties)
 		{
 			if (Difference.Identifier == PropertyIdentifer)
 			{
-				return &Difference;
+				OrderedProperties.Push(&Difference);
+				return true;
 			}
 		}
-		return nullptr;
+		return false;
 	};
 
 	// zip the two sets of properties, zip iterators are not trivial to write in c++,
 	// so this procedural stuff will have to do:
 	int IterOld = 0;
 	int IterNew = 0;
-	while (IterOld != OldProperties.Num() || IterNew != NewProperties.Num())
+	while (IterOld < OldProperties.Num() || IterNew < NewProperties.Num())
 	{
-		if (IterOld != OldProperties.Num() && IterNew != NewProperties.Num()
-			&& OldProperties[IterOld] == NewProperties[IterNew])
+		const bool bOldIterValid = IterOld < OldProperties.Num();
+		const bool bNewIterValid = IterNew < NewProperties.Num();
+
+		// We've reached the end of the new list, but still have properties in the old list.
+		// Continue over the old list to catch any remaining diffs.
+		if (bOldIterValid && !bNewIterValid)
 		{
-			if (const FSingleObjectDiffEntry* Differing = FindDiffering(DifferingProperties, OldProperties[IterOld]))
-			{
-				OrderedProperties.Push(Differing);
-			}
+			FindAndPushDiff(OldProperties[IterOld]);
 			++IterOld;
+		}
+		// We've reached the end of the old list, but still have properties in the new list.
+		// Continue over the new list to catch any remaining diffs.
+		else if (!bOldIterValid && bNewIterValid)
+		{
+			FindAndPushDiff(NewProperties[IterNew]);
 			++IterNew;
 		}
 		else
 		{
-			if (IterOld != OldProperties.Num())
+			// If both properties have the same path, check to ensure the property hasn't changed.
+			if (OldProperties[IterOld] == NewProperties[IterNew])
+			{
+				FindAndPushDiff(OldProperties[IterOld]);
+				++IterNew;
+				++IterOld;
+			}
+			else
 			{
-				bool bFoundDifference = false;
-				if (const FSingleObjectDiffEntry* OldDiffering = FindDiffering(DifferingProperties, OldProperties[IterOld]))
+				// If the old property is different, add it to the list and increment the old iter.
+				// This indicates the property was removed.
+				if (FindAndPushDiff(OldProperties[IterOld]))
 				{
-					bFoundDifference = true;
-					OrderedProperties.Push(OldDiffering);
 					++IterOld;
 				}
-				else if (IterNew != NewProperties.Num())
+				// If the new property is different, add it to the list and increment the new iter.
+				// This indicates the property was added.
+				else if (FindAndPushDiff(NewProperties[IterNew]))
 				{
-					if (const FSingleObjectDiffEntry* NewDiffering = FindDiffering(DifferingProperties, NewProperties[IterNew]))
-					{
-						bFoundDifference = true;
-						OrderedProperties.Push(NewDiffering);
-						++IterNew;
-					}
+					++IterNew;
 				}
-				if (!bFoundDifference)
+				// Neither property was different.
+				// This indicates the iterators were just out of step from a previous addition or removal.
+				else
 				{
 					++IterOld;
 					++IterNew;
 				}
 			}
-			else
-			{
-				check(IterNew != NewProperties.Num());
-				if (const FSingleObjectDiffEntry* Differing = FindDiffering(DifferingProperties, NewProperties[IterNew]))
-				{
-					OrderedProperties.Push(Differing);
-				}
-				++IterNew;
-			}
 		}
 	}
 
diff --git a/Engine/Source/Editor/Kismet/Private/SBlueprintPalette.cpp b/Engine/Source/Editor/Kismet/Private/SBlueprintPalette.cpp
--- a/Engine/Source/Editor/Kismet/Private/SBlueprintPalette.cpp
+++ b/Engine/Source/Editor/Kismet/Private/SBlueprintPalette.cpp
@@ -1248,7 +1248,7 @@ bool SBlueprintPaletteItem::OnNameTextVerifyChanged(const FText& InNewText, FTex
 	{
 		for (USCS_Node* Node : BlueprintObj->SimpleConstructionScript->GetAllNodes())
 		{
-			if (Node && Node->VariableName == OriginalName && !FComponentEditorUtils::IsValidVariableNameString(Node->ComponentTemplate, InNewText.ToString()))
+			if (Node && Node->GetVariableName() == OriginalName && !FComponentEditorUtils::IsValidVariableNameString(Node->ComponentTemplate, InNewText.ToString()))
 			{
 				OutErrorMessage = LOCTEXT("RenameFailed_NotValid", "This name is reserved for engine use.");
 				return false;
diff --git a/Engine/Source/Editor/Kismet/Private/SKismetInspector.cpp b/Engine/Source/Editor/Kismet/Private/SKismetInspector.cpp
--- a/Engine/Source/Editor/Kismet/Private/SKismetInspector.cpp
+++ b/Engine/Source/Editor/Kismet/Private/SKismetInspector.cpp
@@ -222,9 +222,10 @@ FText SKismetInspector::GetContextualEditingWidgetTitle() const
 			{
 				if (SCSNode->ComponentTemplate != NULL)
 				{
-					if (SCSNode->VariableName != NAME_None)
+					const FName VariableName = SCSNode->GetVariableName();
+					if (VariableName != NAME_None)
 					{
-						Title = FText::Format(LOCTEXT("TemplateForFmt", "Template for {0}"), FText::FromName(SCSNode->VariableName));
+						Title = FText::Format(LOCTEXT("TemplateForFmt", "Template for {0}"), FText::FromName(VariableName));
 					}
 					else 
 					{
diff --git a/Engine/Source/Editor/Kismet/Private/SMyBlueprint.cpp b/Engine/Source/Editor/Kismet/Private/SMyBlueprint.cpp
--- a/Engine/Source/Editor/Kismet/Private/SMyBlueprint.cpp
+++ b/Engine/Source/Editor/Kismet/Private/SMyBlueprint.cpp
@@ -2693,6 +2693,10 @@ void SMyBlueprint::OnResetItemFilter()
 void SMyBlueprint::EnsureLastPinTypeValid()
 {
 	const UEdGraphSchema_K2* Schema = GetDefault<UEdGraphSchema_K2>();
+
+	LastPinType.bIsWeakPointer = false;
+	LastFunctionPinType.bIsWeakPointer = false;
+
 	const bool bLastPinTypeValid = (Schema->PC_Struct != LastPinType.PinCategory) || LastPinType.PinSubCategoryObject.IsValid();
 	const bool bLastFunctionPinTypeValid = (Schema->PC_Struct != LastFunctionPinType.PinCategory) || LastFunctionPinType.PinSubCategoryObject.IsValid();
 	const bool bConstType = LastPinType.bIsConst || LastFunctionPinType.bIsConst;
diff --git a/Engine/Source/Editor/Kismet/Private/SSCSEditor.cpp b/Engine/Source/Editor/Kismet/Private/SSCSEditor.cpp
--- a/Engine/Source/Editor/Kismet/Private/SSCSEditor.cpp
+++ b/Engine/Source/Editor/Kismet/Private/SSCSEditor.cpp
@@ -4025,7 +4025,7 @@ FSCSEditorTreeNodePtrType SSCSEditor::GetNodeFromActorComponent(const UActorComp
 							for (USCS_Node* SCS_Node : ParentBPStack[StackIndex]->SimpleConstructionScript->GetAllNodes())
 							{
 								check(SCS_Node != NULL);
-								if (SCS_Node->VariableName == ActorComponent->GetFName())
+								if (SCS_Node->GetVariableName() == ActorComponent->GetFName())
 								{
 									// We found a match; redirect to the component archetype instance that may be associated with a tree node
 									ActorComponent = SCS_Node->ComponentTemplate;
@@ -4735,9 +4735,10 @@ UActorComponent* SSCSEditor::AddNewNode(USCS_Node* NewNode, UObject* Asset, bool
 	NewNodePtr = AddTreeNode(NewNode, SceneRootNodePtr, false);
 
 	// Potentially adjust variable names for any child blueprints
-	if(NewNode->VariableName != NAME_None)
+	const FName VariableName = NewNode->GetVariableName();
+	if(VariableName != NAME_None)
 	{
-		FBlueprintEditorUtils::ValidateBlueprintChildVariables(Blueprint, NewNode->VariableName);
+		FBlueprintEditorUtils::ValidateBlueprintChildVariables(Blueprint, VariableName);
 	}
 	
 	if(bSetFocusToNewItem)
@@ -5221,14 +5222,14 @@ FSCSEditorTreeNodePtrType SSCSEditor::AddTreeNode(USCS_Node* InSCSNode, FSCSEdit
 	{
 		check(InSCSNode->ComponentTemplate != NULL);
 		checkf(InSCSNode->ParentComponentOrVariableName == NAME_None
-			|| (!InSCSNode->bIsParentComponentNative && InParentNodePtr->GetSCSNode() != NULL && InParentNodePtr->GetSCSNode()->VariableName == InSCSNode->ParentComponentOrVariableName)
+			|| (!InSCSNode->bIsParentComponentNative && InParentNodePtr->GetSCSNode() != NULL && InParentNodePtr->GetSCSNode()->GetVariableName() == InSCSNode->ParentComponentOrVariableName)
 			|| (InSCSNode->bIsParentComponentNative && InParentNodePtr->GetComponentTemplate() != NULL && InParentNodePtr->GetComponentTemplate()->GetFName() == InSCSNode->ParentComponentOrVariableName),
 			TEXT("Failed to add SCS node %s to tree:\n- bIsParentComponentNative=%d\n- Stored ParentComponentOrVariableName=%s\n- Actual ParentComponentOrVariableName=%s"),
-			*InSCSNode->VariableName.ToString(),
+			*InSCSNode->GetVariableName().ToString(),
 			!!InSCSNode->bIsParentComponentNative,
 			*InSCSNode->ParentComponentOrVariableName.ToString(),
 			!InSCSNode->bIsParentComponentNative
-			? (InParentNodePtr->GetSCSNode() != NULL ? *InParentNodePtr->GetSCSNode()->VariableName.ToString() : TEXT("NULL"))
+			? (InParentNodePtr->GetSCSNode() != NULL ? *InParentNodePtr->GetSCSNode()->GetVariableName().ToString() : TEXT("NULL"))
 			: (InParentNodePtr->GetComponentTemplate() != NULL ? *InParentNodePtr->GetComponentTemplate()->GetFName().ToString() : TEXT("NULL")));
 	}
 	
diff --git a/Engine/Source/Editor/Kismet/Public/BlueprintEditor.h b/Engine/Source/Editor/Kismet/Public/BlueprintEditor.h
--- a/Engine/Source/Editor/Kismet/Public/BlueprintEditor.h
+++ b/Engine/Source/Editor/Kismet/Public/BlueprintEditor.h
@@ -842,7 +842,7 @@ protected:
 	 * @param InDestinationGraph		The destination graph to move the selected nodes to
 	 * @param InCollapsableNodes		The selection of nodes being collapsed
 	 */
-	void CollapseNodesIntoGraph(UEdGraphNode* InGatewayNode, class UK2Node_EditablePinBase* InEntryNode, class UK2Node_EditablePinBase* InResultNode, UEdGraph* InSourceGraph, UEdGraph* InDestinationGraph, TSet<UEdGraphNode*>& InCollapsableNodes, bool bCanDiscardEmptyReturnNode = false);
+	void CollapseNodesIntoGraph(UEdGraphNode* InGatewayNode, class UK2Node_EditablePinBase* InEntryNode, class UK2Node_EditablePinBase* InResultNode, UEdGraph* InSourceGraph, UEdGraph* InDestinationGraph, TSet<UEdGraphNode*>& InCollapsableNodes, bool bCanDiscardEmptyReturnNode, bool bCanHaveWeakObjPtrParam);
 
 	/** Called when a selection of nodes are being collapsed into a sub-graph */
 	void CollapseNodes(TSet<class UEdGraphNode*>& InCollapsableNodes);
diff --git a/Engine/Source/Editor/KismetCompiler/Private/EdGraphCompiler.cpp b/Engine/Source/Editor/KismetCompiler/Private/EdGraphCompiler.cpp
--- a/Engine/Source/Editor/KismetCompiler/Private/EdGraphCompiler.cpp
+++ b/Engine/Source/Editor/KismetCompiler/Private/EdGraphCompiler.cpp
@@ -135,19 +135,6 @@ UEdGraphNode* FGraphCompilerContext::FindNodeByClass(const UEdGraph* Graph, TSub
 	return FirstResultNode;
 }
 
-
-void FGraphCompilerContext::FindNodesByClass(const UEdGraph* Graph, TSubclassOf<UEdGraphNode>  NodeClass, TArray<UEdGraphNode*>& FoundNodes) const
-{
-	for (int32 NodeIndex = 0; NodeIndex < Graph->Nodes.Num(); ++NodeIndex)
-	{
-		UEdGraphNode* Node = Graph->Nodes[NodeIndex];
-		if (Node && Node->IsA(NodeClass))
-		{
-			FoundNodes.Add(Node);
-		}
-	}
-}
-
 /** Prunes any nodes that weren't visited from the graph, printing out a warning */
 void FGraphCompilerContext::PruneIsolatedNodes(const TArray<UEdGraphNode*>& RootSet, TArray<UEdGraphNode*>& GraphNodes)
 {
diff --git a/Engine/Source/Editor/KismetCompiler/Private/KismetCompiler.cpp b/Engine/Source/Editor/KismetCompiler/Private/KismetCompiler.cpp
--- a/Engine/Source/Editor/KismetCompiler/Private/KismetCompiler.cpp
+++ b/Engine/Source/Editor/KismetCompiler/Private/KismetCompiler.cpp
@@ -27,6 +27,7 @@
 #include "Engine/SCS_Node.h"
 #include "Engine/TimelineTemplate.h"
 #include "Components/TimelineComponent.h"
+#include "ArchiveScriptReferenceCollector.h"
 
 static bool bDebugPropertyPropagation = false;
 
@@ -50,6 +51,43 @@ DECLARE_CYCLE_STAT(TEXT("Bind and Link Class"), EKismetCompilerStats_BindAndLink
 DECLARE_CYCLE_STAT(TEXT("Calculate checksum of CDO"), EKismetCompilerStats_ChecksumCDO, STATGROUP_KismetCompiler );
 DECLARE_CYCLE_STAT(TEXT("Analyze execution path"), EKismetCompilerStats_AnalyzeExecutionPath, STATGROUP_KismetCompiler);
 DECLARE_CYCLE_STAT(TEXT("Calculate checksum of signature"), EKismetCompilerStats_ChecksumSignature, STATGROUP_KismetCompiler);
+
+namespace
+{
+	// The function collects all nodes, that can represents entry points of the execution. Any node connected to "root" node (by execution link) won't be consider isolated.
+	static void GatherRootSet(const UEdGraph* Graph, TArray<UEdGraphNode*>& RootSet, bool bIncludeNodesThatCouldBeExpandedToRootSet)
+	{
+		for (UEdGraphNode* Node : Graph->Nodes)
+		{
+			const bool bRootSetByType = Node && (Node->IsA<UK2Node_FunctionEntry>() || Node->IsA<UK2Node_Event>() || Node->IsA<UK2Node_Timeline>());
+			bool bIsRootSet = bRootSetByType;
+
+			if (Node && bIncludeNodesThatCouldBeExpandedToRootSet && !bIsRootSet)
+			{
+				//Include non-pure K2Nodes, without input pins
+				UK2Node* K2Node = Cast<UK2Node>(Node);
+				auto HasInputPins = [](UK2Node* InNode) -> bool
+				{
+					for (UEdGraphPin* Pin : InNode->Pins)
+					{
+						if (Pin && (EEdGraphPinDirection::EGPD_Input == Pin->Direction))
+						{
+							return true;
+						}
+					}
+					return false;
+				};
+
+				bIsRootSet |= K2Node && !K2Node->IsNodePure() && !HasInputPins(K2Node);
+			}
+			if (bIsRootSet)
+			{
+				RootSet.Add(Node);
+			}
+		}
+	}
+}
+
 //////////////////////////////////////////////////////////////////////////
 // FKismetCompilerContext
 
@@ -637,9 +675,9 @@ void FKismetCompilerContext::CreateClassVariablesFromBlueprint()
 			if (Node)
 			{
 				FName VarName = Node->GetVariableName();
-				if ((VarName != NAME_None) && (Node->ComponentTemplate != NULL))
+				if ((VarName != NAME_None) && (Node->ComponentClass != nullptr))
 				{
-					FEdGraphPinType Type(Schema->PC_Object, TEXT(""), Node->ComponentTemplate->GetClass(), false, false);
+					FEdGraphPinType Type(Schema->PC_Object, TEXT(""), Node->ComponentClass, false, false);
 					UProperty* NewProperty = CreateVariable(VarName, Type);
 					if (NewProperty != NULL)
 					{
@@ -1154,18 +1192,39 @@ void FKismetCompilerContext::PruneIsolatedNodes(const TArray<UEdGraphNode*>& Roo
 		Visitor.TraverseNodes(RootNode);
 	}
 
+	const UEdGraphSchema* const K2Schema = UEdGraphSchema_K2::StaticClass()->GetDefaultObject<UEdGraphSchema_K2>();
+
 	for (int32 NodeIndex = 0; NodeIndex < GraphNodes.Num(); ++NodeIndex)
 	{
 		UEdGraphNode* Node = GraphNodes[NodeIndex];
 		if (!Node || (!Visitor.VisitedNodes.Contains(Node) && !IsNodePure(Node)))
 		{
-			if (!CanIgnoreNode(Node))
+			auto ShouldKeepNonPureNodeWithoutExecPin = [&]() -> bool
 			{
-				// Disabled this warning, because having orphaned chains is standard workflow for LDs
-				//MessageLog.Warning(TEXT("Node @@ will never be executed and is being pruned"), Node);
-			}
+				if (Node 
+					&& Node->CanCreateUnderSpecifiedSchema(K2Schema) // Anim Nodes still should be pruned
+					&& !Node->IsA<UK2Node_Tunnel>()) // Tunnels are never pure.
+				{
+					bool bHasExecPin = false;
+					for (const UEdGraphPin* Pin : Node->Pins)
+					{
+						if (Pin && (Pin->PinType.PinCategory == UEdGraphSchema_K2::PC_Exec))
+						{
+							bHasExecPin = true;
+							break;
+						}
+					}
+					if (!bHasExecPin)
+					{
+						const FString WarningStr = FString::Printf(*LOCTEXT("NoPureNodeWithoutExec_Warning", "Node @@. The node won't be pruned as isolated one. The node is not pure, but it has no exec pin(s). Verify IsNodePure implementation in %s.").ToString(), *Node->GetClass()->GetDisplayNameText().ToString());
+						MessageLog.Warning(*WarningStr, Node);
+					}
+					return !bHasExecPin;
+				}
+				return false;
+			};
 
-			if (!Node || !ShouldForceKeepNode(Node))
+			if (!Node || (!ShouldForceKeepNode(Node) && !ShouldKeepNonPureNodeWithoutExecPin()))
 			{
 				if (Node)
 				{
@@ -1239,28 +1298,31 @@ void FKismetCompilerContext::PrecompileFunction(FKismetFunctionContext& Context)
 	BP_SCOPED_COMPILER_EVENT_STAT(EKismetCompilerStats_PrecompileFunction);
 
 	// Find the root node, which will drive everything else
-	check(Context.RootSet.Num() == 0);
-	FindNodesByClass(Context.SourceGraph, UK2Node_FunctionEntry::StaticClass(), Context.RootSet);
+	TArray<UK2Node_FunctionEntry*> EntryPoints;
+	Context.SourceGraph->GetNodesOfClass(EntryPoints);
 
-	if (Context.RootSet.Num())
+	if (EntryPoints.Num())
 	{
-		Context.EntryPoint = CastChecked<UK2Node_FunctionEntry>(Context.RootSet[0]);
+		Context.EntryPoint = CastChecked<UK2Node_FunctionEntry>(EntryPoints[0]);
 
 		// Make sure there was only one function entry node
-		for (int32 i = 1; i < Context.RootSet.Num(); ++i)
+		for (int32 i = 1; i < EntryPoints.Num(); ++i)
 		{
 			MessageLog.Error(*FString::Printf(*LOCTEXT("ExpectedOneFunctionEntry_Error", "Expected only one function entry node in graph @@, but found both @@ and @@").ToString()),
 				Context.SourceGraph,
 				Context.EntryPoint,
-				Context.RootSet[i]);
+				EntryPoints[i]);
 		}
 
-		// Find any other entry points caused by special nodes
-		FindNodesByClass(Context.SourceGraph, UK2Node_Event::StaticClass(), Context.RootSet);
-		FindNodesByClass(Context.SourceGraph, UK2Node_Timeline::StaticClass(), Context.RootSet);
+		{
+			TArray<UEdGraphNode*> RootSet;
+			const bool bIncludePotentialRootNodes = false;
+			// Find any all entry points caused by special nodes
+			GatherRootSet(Context.SourceGraph, RootSet, bIncludePotentialRootNodes);
 
-		// Find the connected subgraph starting at the root node and prune out unused nodes
-		PruneIsolatedNodes(Context.RootSet, Context.SourceGraph->Nodes);
+			// Find the connected subgraph starting at the root node and prune out unused nodes
+			PruneIsolatedNodes(RootSet, Context.SourceGraph->Nodes);
+		}
 
 		if (bIsFullCompile)
 		{
@@ -1268,14 +1330,6 @@ void FKismetCompilerContext::PrecompileFunction(FKismetFunctionContext& Context)
 			ValidateSelfPinsInGraph(Context);
 			ValidateNoWildcardPinsInGraph(Context.SourceGraph);
 
-			for (int32 ChildIndex = 0; ChildIndex < Context.SourceGraph->Nodes.Num(); ++ChildIndex)
-			{
-				if (const UK2Node* K2Node = Cast<const UK2Node>(Context.SourceGraph->Nodes[ChildIndex]))
-				{
-					K2Node->ValidateNodeAfterPrune(Context.MessageLog);
-				}
-			}
-
 			// Transforms
 			TransformNodes(Context);
 		}
@@ -1596,6 +1650,7 @@ void FKismetCompilerContext::CompileFunction(FKismetFunctionContext& Context)
 		{
 			if (!bPureNode)
 			{
+				UEdGraphPin* ExecPin = nullptr;
 				bool bEmitDebuggingSite = true;
 
 				if (Context.IsEventGraph() && (Node->IsA(UK2Node_FunctionEntry::StaticClass())))
@@ -1605,25 +1660,27 @@ void FKismetCompilerContext::CompileFunction(FKismetFunctionContext& Context)
 					// debug data for the behind-the-scenes entry point, only for the user-visible ones.
 					bEmitDebuggingSite = false;
 				}
-				if (bEmitDebuggingSite)
+				if (Context.bInstrumentScriptCode)
 				{
-					FBlueprintCompiledStatement& Statement = Context.AppendStatementForNode(Node);
-					Statement.Type = Context.GetBreakpointType();
-					if (Context.IsInstrumentationRequired())
+					for (auto Pin : Node->Pins)
 					{
-						for (auto Pin : Node->Pins)
+						if (Pin->Direction == EGPD_Input && Pin->PinType.PinCategory == UEdGraphSchema_K2::PC_Exec)
 						{
-							if (Pin->Direction == EGPD_Input && Pin->PinType.PinCategory == UEdGraphSchema_K2::PC_Exec)
-							{
-								Statement.ExecContext = Pin;
-								break;
-							}
+							ExecPin = Pin;
+							bEmitDebuggingSite = Context.IsInstrumentationRequiredForPin(ExecPin);
+							break;
 						}
 					}
+				}
+				if (bEmitDebuggingSite)
+				{
+					FBlueprintCompiledStatement& Statement = Context.AppendStatementForNode(Node);
+					Statement.Type = Context.GetBreakpointType();
+					Statement.ExecContext = ExecPin;
 					Statement.Comment = NodeComment;
 				}
 			}
-			else if (Context.IsInstrumentationRequired())
+			else if (Context.IsInstrumentationRequiredForNode(Node))
 			{
 				FBlueprintCompiledStatement& Statement = Context.AppendStatementForNode(Node);
 				Statement.Type = KCST_InstrumentedPureNodeEntry;
@@ -1689,7 +1746,8 @@ void FKismetCompilerContext::CompileFunction(FKismetFunctionContext& Context)
 						for (UEdGraphPin* LinkedTo : Pin->LinkedTo)
 						{
 							UEdGraphNode* NodeUsingOutput = LinkedTo->GetOwningNode();
-							if (NodeUsingOutput != nullptr)
+							const bool bValidExpansionLink = !Context.bInstrumentScriptCode || Context.IsInstrumentationRequiredForPin(LinkedTo);
+							if (NodeUsingOutput != nullptr && bValidExpansionLink)
 							{
 								// Add this node, as well as other nodes this node depends on
 								TSet<UEdGraphNode*>& TargetNodesRequired = PureNodesNeeded.FindOrAdd(NodeUsingOutput);
@@ -1721,7 +1779,7 @@ void FKismetCompilerContext::CompileFunction(FKismetFunctionContext& Context)
 					OrderedInsertIntoArray(SortedPureNodes, SortKeyMap, *It);
 				}
 
-				if (Context.IsInstrumentationRequired())
+				if (Context.bInstrumentScriptCode)
 				{
 					FBlueprintCompiledStatement& PopState = Context.PrependStatementForNode(Node);
 					PopState.Type = KCST_InstrumentedStatePop;
@@ -1735,7 +1793,7 @@ void FKismetCompilerContext::CompileFunction(FKismetFunctionContext& Context)
 					Context.CopyAndPrependStatements(Node, NodeToInline);
 				}
 
-				if (Context.IsInstrumentationRequired())
+				if (Context.bInstrumentScriptCode)
 				{
 					FBlueprintCompiledStatement& PushState = Context.PrependStatementForNode(Node);
 					PushState.Type = KCST_InstrumentedStatePush;
@@ -2057,12 +2115,12 @@ void FKismetCompilerContext::CreatePinEventNodeForTimelineFunction(UK2Node_Timel
 	TimelineEventNode->AllocateDefaultPins();
 
 	// Move any links from 'update' pin to the 'update event' node
-	UEdGraphPin* UpdatePin = TimelineNode->FindPin(PinName);
-	check(UpdatePin);
+	UEdGraphPin* UpdatePin = TimelineNode ? TimelineNode->FindPin(PinName) : nullptr;
+	ensureMsgf(UpdatePin, TEXT("Timeline '%s' has no pin '%s'"), *GetPathNameSafe(TimelineNode), *PinName);
 
 	UEdGraphPin* UpdateOutput = Schema->FindExecutionPin(*TimelineEventNode, EGPD_Output);
 
-	if(UpdatePin != NULL && UpdateOutput != NULL)
+	if(UpdatePin && UpdateOutput)
 	{
 		MovePinLinksToIntermediate(*UpdatePin, *UpdateOutput);
 	}
@@ -2522,6 +2580,17 @@ void FKismetCompilerContext::MergeUbergraphPagesIn(UEdGraph* Ubergraph)
 // Expands out nodes that need it
 void FKismetCompilerContext::ExpansionStep(UEdGraph* Graph, bool bAllowUbergraphExpansions)
 {
+	auto PruneInner = [=]()
+	{
+		TArray<UEdGraphNode*> RootSet;
+		const bool bIncludePotentialRootNodes = true;
+		// Find any all entry points caused by special nodes
+		GatherRootSet(Graph, RootSet, bIncludePotentialRootNodes);
+
+		// Find the connected subgraph starting at the root node and prune out unused nodes
+		PruneIsolatedNodes(RootSet, Graph->Nodes);
+	};
+
 	// Node expansion may affect the signature of a static function
 	if (bIsFullCompile)
 	{
@@ -2530,15 +2599,56 @@ void FKismetCompilerContext::ExpansionStep(UEdGraph* Graph, bool bAllowUbergraph
 		// Collapse any remaining tunnels or macros
 		ExpandTunnelsAndMacros(Graph);
 
+		// First pruning pass must be call after all collapsed nodes are expanded. Before the expansion we don't know which collapsed graph is really isolated. 
+		// If the pruning was called before expansion (and all collapsed graphs were saved), the isolated collapsed graphs would be unnecessarily validated.
+		PruneInner();
+
 		for (int32 NodeIndex = 0; NodeIndex < Graph->Nodes.Num(); ++NodeIndex)
 		{
 			UK2Node* Node = Cast<UK2Node>(Graph->Nodes[NodeIndex]);
 			if (Node)
 			{
-				Node->ExpandNode(*this, Graph);
+				if (!CompileOptions.bAddInstrumentation)
+				{
+					Node->ExpandNode(*this, Graph);
+				}
+				else
+				{
+					// Determine links before expansion
+					TMap<UEdGraphPin*, UEdGraphPin*> SourceNodeLinks;
+					DetermineNodeExecLinks(Node, SourceNodeLinks);
+					// track the current node count so we can figure out if expansion happened.
+					int32 ExpansionNodeIdx = Graph->Nodes.Num();
+					Node->ExpandNode(*this, Graph);
+					// Keep track of new expansion nodes
+					if (ExpansionNodeIdx < Graph->Nodes.Num())
+					{
+						TArray<UEdGraphNode*> ExpansionNodes;
+						for (; ExpansionNodeIdx < Graph->Nodes.Num(); ++ExpansionNodeIdx)
+						{
+							ExpansionNodes.Add(Graph->Nodes[ExpansionNodeIdx]);
+							MessageLog.RegisterExpansionNode(Graph->Nodes[ExpansionNodeIdx]);
+							// Add the relevent pins to the trace suppression filter
+							for (auto ExpansionPin : Graph->Nodes[ExpansionNodeIdx]->Pins)
+							{
+								for (auto ExpansionPinLink : ExpansionPin->LinkedTo)
+								{
+									MessageLog.AddPinTraceFilter(ExpansionPin);
+								}
+							}
+						}
+						// Create boundaries around expsnion nodes so we can simulate pin traces correctly
+						UK2Node_TunnelBoundary::CreateBoundariesForExpansionNodes(Node, ExpansionNodes, SourceNodeLinks, MessageLog);
+						ExpansionNodeIdx = Graph->Nodes.Num();
+					}
+				}
 			}
 		}
 	}
+	else
+	{
+		PruneInner();
+	}
 
 	if (bAllowUbergraphExpansions)
 	{
@@ -2547,6 +2657,22 @@ void FKismetCompilerContext::ExpansionStep(UEdGraph* Graph, bool bAllowUbergraph
 	}
 }
 
+void FKismetCompilerContext::DetermineNodeExecLinks(UEdGraphNode* SourceNode, TMap<UEdGraphPin*, UEdGraphPin*>& SourceNodeLinks) const
+{
+	// Find all linked pins we care about from the source node.
+	for (auto SourcePin : SourceNode->Pins)
+	{
+		if (SourcePin->PinType.PinCategory == UEdGraphSchema_K2::PC_Exec)
+		{
+			UEdGraphPin* TrueSourcePin = MessageLog.FindSourcePin(SourcePin);
+			for (auto LinkedPin : SourcePin->LinkedTo)
+			{
+				SourceNodeLinks.Add(LinkedPin, TrueSourcePin);
+			}
+		}
+	}
+}
+
 void FKismetCompilerContext::VerifyValidOverrideEvent(const UEdGraph* Graph)
 {
 	check(NULL != Graph);
@@ -2832,6 +2958,15 @@ void FKismetCompilerContext::ExpandTunnelsAndMacros(UEdGraph* SourceGraph)
 			continue;
 		}
 
+		UK2Node_Tunnel* TunnelNode = Cast<UK2Node_Tunnel>(*NodeIt);
+		// After this expansion (and before the validation) PruneIsolatedNodes is called. So this is the last chance to validate nodes like UK2Node_MathExpression.
+		// Notice: even isolated MathExpression nodes will be validated. But, since the MathExpression is usually optimized (so it is not handled here as tunnel, because ShouldMergeChildGraphs return false) it is not a problem.
+		// Notice: MacroInstance Node is based on Tunnel Node.
+		if (TunnelNode)
+		{
+			TunnelNode->ValidateNodeDuringCompilation(MessageLog);
+		}
+
 		if (UK2Node_MacroInstance* MacroInstanceNode = Cast<UK2Node_MacroInstance>(*NodeIt))
 		{
 			UEdGraph* MacroGraph = MacroInstanceNode->GetMacroGraph();
@@ -2853,6 +2988,16 @@ void FKismetCompilerContext::ExpandTunnelsAndMacros(UEdGraph* SourceGraph)
 			// Clone the macro graph, then move all of its children, keeping a list of nodes from the macro
 			UEdGraph* ClonedGraph = FEdGraphUtilities::CloneGraph(MacroGraph, NULL, &MessageLog, true);
 
+			// Find active tunnel instances for this macro instance.
+			TArray<TWeakObjectPtr<UEdGraphNode>> TunnelInstances;
+			MessageLog.GetTunnelsActiveForNode(MacroInstanceNode, TunnelInstances);
+			if (!TunnelInstances.Num())
+			{
+				UEdGraphNode* SourceMacroInstance = Cast<UEdGraphNode>(MessageLog.FindSourceObject(MacroInstanceNode));
+				TunnelInstances.Add(SourceMacroInstance);
+				MessageLog.RegisterIntermediateTunnelInstance(MacroInstanceNode, TunnelInstances);
+			}
+
 			if (CompileOptions.IsInstrumentationActive())
 			{
 				// Create tunnel Boundary signals
@@ -2974,22 +3119,13 @@ void FKismetCompilerContext::ExpandTunnelsAndMacros(UEdGraph* SourceGraph)
 						DuplicatedNode->ReconstructNode();
 					}
 
-					// Record the source node mapping for the intermediate node first, as it's going to be overwritten through the MessageLog below
-					UEdGraphNode* MacroSourceNode = Cast<UEdGraphNode>(MessageLog.FindSourceObject(DuplicatedNode));
-					if (MacroSourceNode)
+					// Record the source node mapping for the intermediate node and register any active tunnel instances if it is a tunnel/macro instance node.
+					MessageLog.RegisterIntermediateTunnelNode(DuplicatedNode, MacroInstanceNode);
+					if (FBlueprintEditorUtils::IsTunnelInstanceNode(DuplicatedNode))
 					{
-						MessageLog.FinalNodeBackToMacroSourceMap.NotifyIntermediateObjectCreation(DuplicatedNode, MacroSourceNode);
-
-						// Also record mappings from final macro source node to intermediate macro instance nodes (there may be more than one)
-						UEdGraphNode* MacroInstanceSourceNode = Cast<UEdGraphNode>(MessageLog.FinalNodeBackToMacroSourceMap.FindSourceObject(MacroInstanceNode));
-						if (MacroInstanceSourceNode && MacroInstanceSourceNode != MacroInstanceNode)
-						{
-							MessageLog.MacroSourceToMacroInstanceNodeMap.Add(MacroSourceNode, MacroInstanceSourceNode);
-						}
+						MessageLog.RegisterIntermediateTunnelInstance(DuplicatedNode, TunnelInstances);
 					}
 
-					MessageLog.NotifyIntermediateObjectCreation(DuplicatedNode, MacroInstanceNode);
-
 					DuplicatedNode->NodePosY += NodeOffsetY;
 					DuplicatedNode->NodePosX += NodeOffsetX;
 
@@ -3000,31 +3136,31 @@ void FKismetCompilerContext::ExpandTunnelsAndMacros(UEdGraph* SourceGraph)
 						continue;
 					}
 					
-					if (UK2Node_Tunnel* TunnelNode = Cast<UK2Node_Tunnel>(DuplicatedNode))
+					if (UK2Node_Tunnel* DuplicatedTunnelNode = Cast<UK2Node_Tunnel>(DuplicatedNode))
 					{
 						// Tunnel nodes should be connected to the MacroInstance they have been instantiated by.. Note 
 						// that if there are tunnel nodes internal to the macro instance they will be incorrectly 
 						// connected to the MacroInstance.
-						if (TunnelNode->bCanHaveInputs)
+						if (DuplicatedTunnelNode->bCanHaveInputs)
 						{
-							check(!TunnelNode->bCanHaveOutputs);
+							check(!DuplicatedTunnelNode->bCanHaveOutputs);
 							// If this check fails it indicates that we've failed to identify all uses of tunnel nodes and 
 							// are erroneously connecting tunnels to the macro instance when they should be left untouched..
-							check(!TunnelNode->InputSinkNode);
-							TunnelNode->InputSinkNode = MacroInstanceNode;
-							MacroInstanceNode->OutputSourceNode = TunnelNode;
+							check(!DuplicatedTunnelNode->InputSinkNode);
+							DuplicatedTunnelNode->InputSinkNode = MacroInstanceNode;
+							MacroInstanceNode->OutputSourceNode = DuplicatedTunnelNode;
 						}
-						else if (TunnelNode->bCanHaveOutputs)
+						else if (DuplicatedTunnelNode->bCanHaveOutputs)
 						{
-							check(!TunnelNode->OutputSourceNode);
-							TunnelNode->OutputSourceNode = MacroInstanceNode;
-							MacroInstanceNode->InputSinkNode = TunnelNode;
+							check(!DuplicatedTunnelNode->OutputSourceNode);
+							DuplicatedTunnelNode->OutputSourceNode = MacroInstanceNode;
+							MacroInstanceNode->InputSinkNode = DuplicatedTunnelNode;
 						}
 					}
 				}
 			}
 		}
-		else if (UK2Node_Tunnel* TunnelNode = Cast<UK2Node_Tunnel>(*NodeIt))
+		else if (TunnelNode)
 		{
 			UEdGraphNode* InputSink = TunnelNode->GetInputSink();
 			for (UEdGraphPin const* TunnelPin : TunnelNode->Pins)
@@ -3087,29 +3223,41 @@ void FKismetCompilerContext::ProcessOneFunctionGraph(UEdGraph* SourceGraph, bool
 	// If a function in the graph cannot be overridden/placed as event make sure that it is not.
 	VerifyValidOverrideFunction(FunctionGraph);
 
-	FKismetFunctionContext& Context = *new (FunctionList)FKismetFunctionContext(MessageLog, Schema, NewClass, Blueprint, CompileOptions.DoesRequireCppCodeGeneration(), CompileOptions.IsInstrumentationActive());
-	Context.SourceGraph = FunctionGraph;
-
-	if(FBlueprintEditorUtils::IsDelegateSignatureGraph(SourceGraph))
+	// First do some cursory validation (pin types match, inputs to outputs, pins never point to their parent node, etc...)
+	// If this fails we don't proceed any further to avoid crashes or infinite loops
+	// When compiling only the skeleton class, we want the UFunction to be generated and processed so it contains all the local variables, this is unsafe to do during any other compilation mode
+	//
+	// NOTE: the order of this conditional check is intentional, and should not
+	//       be rearranged; we do NOT want ValidateGraphIsWellFormed() ran for 
+	//       skeleton-only compiles (that's why we have that check second) 
+	//       because it would most likely result in errors (the function hasn't
+	//       been added to the class yet, etc.)
+	if ((CompileOptions.CompileType == EKismetCompileType::SkeletonOnly) || ValidateGraphIsWellFormed(FunctionGraph))
 	{
-		Context.SetDelegateSignatureName(SourceGraph->GetFName());
-	}
+		FKismetFunctionContext& Context = *new (FunctionList)FKismetFunctionContext(MessageLog, Schema, NewClass, Blueprint, CompileOptions.DoesRequireCppCodeGeneration(), CompileOptions.IsInstrumentationActive());
+		Context.SourceGraph = FunctionGraph;
 
-	// If this is an interface blueprint, mark the function contexts as stubs
-	if (FBlueprintEditorUtils::IsInterfaceBlueprint(Blueprint))
-	{
-		Context.MarkAsInterfaceStub();
-	}
+		if(FBlueprintEditorUtils::IsDelegateSignatureGraph(SourceGraph))
+		{
+			Context.SetDelegateSignatureName(SourceGraph->GetFName());
+		}
 
-	bool bEnforceConstCorrectness = true;
-	if (FBlueprintEditorUtils::IsBlueprintConst(Blueprint) || Schema->IsConstFunctionGraph(Context.SourceGraph, &bEnforceConstCorrectness))
-	{
-		Context.MarkAsConstFunction(bEnforceConstCorrectness);
-	}
+		// If this is an interface blueprint, mark the function contexts as stubs
+		if (FBlueprintEditorUtils::IsInterfaceBlueprint(Blueprint))
+		{
+			Context.MarkAsInterfaceStub();
+		}
 
-	if ( bInternalFunction )
-	{
-		Context.MarkAsInternalOrCppUseOnly();
+		bool bEnforceConstCorrectness = true;
+		if (FBlueprintEditorUtils::IsBlueprintConst(Blueprint) || Schema->IsConstFunctionGraph(Context.SourceGraph, &bEnforceConstCorrectness))
+		{
+			Context.MarkAsConstFunction(bEnforceConstCorrectness);
+		}
+
+		if ( bInternalFunction )
+		{
+			Context.MarkAsInternalOrCppUseOnly();
+		}
 	}
 }
 
@@ -3303,7 +3451,11 @@ void FKismetCompilerContext::Compile()
 	}
 
 	// Declare class wants profiling data.
-	TargetClass->bHasInstrumentation = CompileOptions.IsInstrumentationActive();
+	if (CompileOptions.IsInstrumentationActive())
+	{
+		TargetClass->bHasInstrumentation = true;
+		MessageLog.bTreatCompositeGraphsAsTunnels =	true;
+	}
 
 	// Early validation
 	if (CompileOptions.CompileType == EKismetCompileType::Full)
@@ -3629,10 +3781,10 @@ void FKismetCompilerContext::Compile()
 			GConfig->GetBool(TEXT("Kismet"), TEXT("CompileDisplaysBinaryBackend"), /*out*/ bDisplayBytecode, GEngineIni);
 		}
 
+		const bool bGenerateStubsOnly = !bIsFullCompile || (0 != MessageLog.NumErrors);
 		// Always run the VM backend, it's needed for more than just debug printing
 		{
 			BP_SCOPED_COMPILER_EVENT_STAT(EKismetCompilerStats_CodeGenerationTime);
-			const bool bGenerateStubsOnly = !bIsFullCompile || (0 != MessageLog.NumErrors);
 			Backend_VM.GenerateCodeFromClass(NewClass, FunctionList, bGenerateStubsOnly);
 			if (!bGenerateStubsOnly)
 			{
@@ -3640,6 +3792,24 @@ void FKismetCompilerContext::Compile()
 			}
 		}
 
+		if (!bGenerateStubsOnly)
+		{
+			for (FKismetFunctionContext& FunctionContext : FunctionList)
+			{
+				if (FunctionContext.IsValid())
+				{
+					UFunction* Function = FunctionContext.Function; 
+					ensure(0 == Function->ScriptObjectReferences.Num());
+					FArchiveScriptReferenceCollector ObjRefCollector(Function->ScriptObjectReferences);
+
+					for (int32 iCode = 0; iCode < Function->Script.Num();)
+					{
+						Function->SerializeExpr(iCode, ObjRefCollector);
+					}
+				}
+			}
+		}
+
 		if (bDisplayBytecode && bIsFullCompile)
 		{
 			TGuardValue<ELogTimes::Type> DisableLogTimes(GPrintLogTimes, ELogTimes::None);
@@ -4046,4 +4216,4 @@ FString FKismetCompilerContext::GetGuid(const UEdGraphNode* Node) const
 }
 
 #undef LOCTEXT_NAMESPACE
-//////////////////////////////////////////////////////////////////////////
\ No newline at end of file
+//////////////////////////////////////////////////////////////////////////
diff --git a/Engine/Source/Editor/KismetCompiler/Private/KismetCompilerMisc.cpp b/Engine/Source/Editor/KismetCompiler/Private/KismetCompilerMisc.cpp
--- a/Engine/Source/Editor/KismetCompiler/Private/KismetCompilerMisc.cpp
+++ b/Engine/Source/Editor/KismetCompiler/Private/KismetCompilerMisc.cpp
@@ -1407,7 +1407,7 @@ FBlueprintCompiledStatement& FNodeHandlingFunctor::GenerateSimpleThenGoto(FKisme
 		TargetNode = ThenExecPin->LinkedTo[0]->GetOwningNode();
 	}
 
-	if (Context.bCreateDebugData || Context.bInstrumentScriptCode)
+	if (Context.bCreateDebugData && (!Context.bInstrumentScriptCode || Context.IsInstrumentationRequiredForPin(ThenExecPin)))
 	{
 		FBlueprintCompiledStatement& TraceStatement = Context.AppendStatementForNode(&Node);
 		TraceStatement.Type = Context.GetWireTraceType();
@@ -1728,8 +1728,10 @@ void FKismetFunctionContext::ResolveGotoFixups()
 			{
 				continue;
 			}
-
-			InsertWireTrace(GotoIt.Key(), GotoIt.Value());
+			if (!bInstrumentScriptCode || IsInstrumentationRequiredForPin(GotoIt->Value))
+			{
+				InsertWireTrace(GotoIt.Key(), GotoIt.Value());
+			}
 		}
 	}
 
diff --git a/Engine/Source/Editor/KismetCompiler/Private/KismetCompilerVMBackend.cpp b/Engine/Source/Editor/KismetCompiler/Private/KismetCompilerVMBackend.cpp
--- a/Engine/Source/Editor/KismetCompiler/Private/KismetCompilerVMBackend.cpp
+++ b/Engine/Source/Editor/KismetCompiler/Private/KismetCompilerVMBackend.cpp
@@ -1607,12 +1607,12 @@ public:
 		if (SourceNode != NULL)
 		{
 			// Record where this NOP is
-			UEdGraphNode* TrueSourceNode = Cast<UEdGraphNode>(FunctionContext.MessageLog.FindSourceObject(SourceNode));
+			UEdGraphNode* TrueSourceNode = FunctionContext.MessageLog.GetSourceNode(SourceNode);
 			if (TrueSourceNode)
 			{
 				// If this is a debug site for an expanded macro instruction, there should also be a macro source node associated with it
-				UEdGraphNode* MacroSourceNode = Cast<UEdGraphNode>(CompilerContext.MessageLog.FinalNodeBackToMacroSourceMap.FindSourceObject(SourceNode));
-				if (MacroSourceNode == SourceNode)
+				UEdGraphNode* MacroSourceNode = CompilerContext.MessageLog.GetSourceTunnelNode(SourceNode);
+				if (MacroSourceNode == TrueSourceNode)
 				{
 					// The function above will return the given node if not found in the map. In that case there is no associated source macro node, so we clear it.
 					MacroSourceNode = NULL;
@@ -1637,36 +1637,15 @@ public:
 					}
 
 					// Gather up all the macro instance nodes that lead to this macro source node
-					CompilerContext.MessageLog.MacroSourceToMacroInstanceNodeMap.MultiFind(MacroSourceNode, MacroInstanceNodes);
-				}
-
-				// Register source tunnels for instrumentation.
-				if (FunctionContext.IsInstrumentationRequired())
-				{
-					// After blueprint profiler MVP these changes will be refactored and rolled back into normal compilation.
-
-					// Locate the source tunnel instance
-					if (TWeakObjectPtr<UEdGraphNode>* SourceTunnelInstance = CompilerContext.MessageLog.SourceNodeToTunnelInstanceNodeMap.Find(SourceNode))
+					UEdGraphNode* IntermediateMacroInstance = CompilerContext.MessageLog.GetIntermediateTunnelInstance(SourceNode);
+					CompilerContext.MessageLog.GetTunnelsActiveForNode(IntermediateMacroInstance, MacroInstanceNodes);
+					if (MacroInstanceNodes.Num())
 					{
-						if (SourceTunnelInstance->IsValid())
-						{
-							TrueSourceNode = SourceTunnelInstance->Get();
-						}
+						TrueSourceNode = MacroInstanceNodes[0].Get();
 					}
-					// Locate the real source node, as the code above fails with nested tunnels
-					if (TWeakObjectPtr<UEdGraphNode>* NewSourceNode = CompilerContext.MessageLog.IntermediateTunnelNodeToSourceNodeMap.Find(SourceNode))
-					{
-						if (NewSourceNode->IsValid())
-						{
-							MacroSourceNode = NewSourceNode->Get();
-						}
-					}
-					ClassBeingBuilt->GetDebugData().RegisterNodeToCodeAssociation(TrueSourceNode, MacroSourceNode, MacroInstanceNodes, FunctionContext.Function, Offset, bBreakpointSite);
-				}
-				else
-				{
-					ClassBeingBuilt->GetDebugData().RegisterNodeToCodeAssociation(TrueSourceNode, MacroSourceNode, MacroInstanceNodes, FunctionContext.Function, Offset, bBreakpointSite);
 				}
+				// Register the debug information for the node.
+				ClassBeingBuilt->GetDebugData().RegisterNodeToCodeAssociation(TrueSourceNode, MacroSourceNode, MacroInstanceNodes, FunctionContext.Function, Offset, bBreakpointSite);
 
 				// Track pure node script code range for the current impure (exec) node
 				if (Statement.Type == KCST_InstrumentedPureNodeEntry)
diff --git a/Engine/Source/Editor/KismetCompiler/Public/EdGraphCompilerUtilities.h b/Engine/Source/Editor/KismetCompiler/Public/EdGraphCompilerUtilities.h
--- a/Engine/Source/Editor/KismetCompiler/Public/EdGraphCompilerUtilities.h
+++ b/Engine/Source/Editor/KismetCompiler/Public/EdGraphCompilerUtilities.h
@@ -62,11 +62,6 @@ protected:
 	 */
 	UEdGraphNode* FindNodeByClass(const UEdGraph* Graph, TSubclassOf<UEdGraphNode>  NodeClass, bool bExpectedUnique) const;
 
-	/**
-	 * Scans a graph for all nodes of the specified class.
-	 */
-	void FindNodesByClass(const UEdGraph* Graph, TSubclassOf<UEdGraphNode>  NodeClass, TArray<UEdGraphNode*>& FoundNodes) const;
-
 	/** Prunes any nodes that weren't visited from the graph, printing out a warning */
 	virtual void PruneIsolatedNodes(const TArray<UEdGraphNode*>& RootSet, TArray<UEdGraphNode*>& GraphNodes);
 
diff --git a/Engine/Source/Editor/KismetCompiler/Public/KismetCompiledFunctionContext.h b/Engine/Source/Editor/KismetCompiler/Public/KismetCompiledFunctionContext.h
--- a/Engine/Source/Editor/KismetCompiler/Public/KismetCompiledFunctionContext.h
+++ b/Engine/Source/Editor/KismetCompiler/Public/KismetCompiledFunctionContext.h
@@ -39,9 +39,6 @@ public:
 	// The nominal function entry point
 	UK2Node_FunctionEntry* EntryPoint;
 
-	// The root set of nodes for the purposes of reachability, etc...
-	TArray<UEdGraphNode*> RootSet;
-
 	UFunction* Function;
 	UBlueprintGeneratedClass* NewClass;
 
@@ -183,6 +180,26 @@ public:
 		return bInstrumentScriptCode;
 	}
 
+	bool IsInstrumentationRequiredForNode(UEdGraphNode* Node) const
+	{
+		return bInstrumentScriptCode && !IsExpansionNode(Node);
+	}
+
+	bool IsInstrumentationRequiredForPin(UEdGraphPin* Pin) const
+	{
+		return bInstrumentScriptCode && !IsPinTraceSuppressed(Pin);
+	}
+
+	bool IsExpansionNode(const UEdGraphNode* Node) const
+	{
+		return MessageLog.IsExpansionNode(Node);
+	}
+
+	bool IsPinTraceSuppressed(const UEdGraphPin* Pin) const
+	{
+		return MessageLog.IsPinTraceSuppressed(Pin);
+	}
+
 	EKismetCompiledStatementType GetWireTraceType() const
 	{
 		return bInstrumentScriptCode ? KCST_InstrumentedWireExit : KCST_WireTraceSite;
diff --git a/Engine/Source/Editor/KismetCompiler/Public/KismetCompiler.h b/Engine/Source/Editor/KismetCompiler/Public/KismetCompiler.h
--- a/Engine/Source/Editor/KismetCompiler/Public/KismetCompiler.h
+++ b/Engine/Source/Editor/KismetCompiler/Public/KismetCompiler.h
@@ -425,6 +425,9 @@ protected:
 	/** Validates the generated class */
 	virtual bool ValidateGeneratedClass(UBlueprintGeneratedClass* Class);
 
+	/** Discovers exec pin links for the sourcenode */
+	void DetermineNodeExecLinks(UEdGraphNode* SourceNode, TMap<UEdGraphPin*, UEdGraphPin*>& SourceNodeLinks) const;
+
 private:
 	/**
 	 * Handles creating a new event node for a given output on a timeline node utilizing the named function
diff --git a/Engine/Source/Editor/KismetWidgets/Private/SPinTypeSelector.cpp b/Engine/Source/Editor/KismetWidgets/Private/SPinTypeSelector.cpp
--- a/Engine/Source/Editor/KismetWidgets/Private/SPinTypeSelector.cpp
+++ b/Engine/Source/Editor/KismetWidgets/Private/SPinTypeSelector.cpp
@@ -1,13 +1,16 @@
 // Copyright 1998-2016 Epic Games, Inc. All Rights Reserved.
 #include "KismetWidgetsPrivatePCH.h"
+
 #include "UnrealEd.h"
+#include "BlueprintEditorUtils.h"
 #include "ClassIconFinder.h"
-#include "SPinTypeSelector.h"
-#include "IDocumentation.h"
 #include "Editor/UnrealEd/Public/SListViewSelectorDropdownMenu.h"
+#include "IDocumentation.h"
+#include "Framework/Notifications/NotificationManager.h"
+#include "SPinTypeSelector.h"
 #include "SSearchBox.h"
 #include "SSubMenuHandler.h"
-#include "BlueprintEditorUtils.h"
+#include "Widgets/Notifications/SNotificationList.h"
 
 #define LOCTEXT_NAMESPACE "PinTypeSelector"
 
@@ -131,6 +134,11 @@ enum class EPinContainerType
 	Map,
 };
 
+static bool ContainerRequiresGetTypeHash(EPinContainerType InType)
+{
+	return InType == EPinContainerType::Set || InType == EPinContainerType::Map;
+}
+
 TSharedRef<SWidget> SPinTypeSelector::ConstructPinTypeImage(const FSlateBrush* PrimaryIcon, const FSlateColor& PrimaryColor, const FSlateBrush* SecondaryIcon, const FSlateColor& SecondaryColor, TSharedPtr<SToolTip> InToolTip)
 {
 	return
@@ -152,6 +160,20 @@ void SPinTypeSelector::Construct(const FArguments& InArgs, FGetPinTypeTree GetPi
 		PinTypes.Add(MakeShareable(new EPinContainerType(EPinContainerType::Map)));
 	}
 
+	static const FSlateBrush* Images[] = {
+		FEditorStyle::GetBrush(TEXT("Kismet.VariableList.TypeIcon")),
+		FEditorStyle::GetBrush(TEXT("Kismet.VariableList.ArrayTypeIcon")),
+		FEditorStyle::GetBrush(TEXT("Kismet.VariableList.SetTypeIcon")),
+		FEditorStyle::GetBrush(TEXT("Kismet.VariableList.MapKeyTypeIcon")),
+	};
+
+	static const FText Tooltips[] = {
+		LOCTEXT("SingleVariableTooltip", "Single Variable"),
+		LOCTEXT("ArrayTooltip", "Array"),
+		LOCTEXT("SetTooltip", "Set"),
+		LOCTEXT("MapTooltip", "Map (Dictionary)"),
+	};
+
 	SearchText = FText::GetEmpty();
 
 	OnTypeChanged = InArgs._OnPinTypeChanged;
@@ -204,60 +226,63 @@ void SPinTypeSelector::Construct(const FArguments& InArgs, FGetPinTypeTree GetPi
 		TSharedPtr<SWidget> ContainerControl;
 		if (FBlueprintEditorUtils::ShouldEnableAdvancedContainers())
 		{
-			auto GenerateContainerTypeEntry = [this](TSharedPtr<EPinContainerType> InPinContainerType)
-			{
-				EPinContainerType PinContainerType = *InPinContainerType;
-
-				static const FSlateBrush* Images[] = {
-					FEditorStyle::GetBrush(TEXT("Kismet.VariableList.TypeIcon")),
-					FEditorStyle::GetBrush(TEXT("Kismet.VariableList.ArrayTypeIcon")),
-					FEditorStyle::GetBrush(TEXT("Kismet.VariableList.SetTypeIcon")),
-					FEditorStyle::GetBrush(TEXT("Kismet.VariableList.MapKeyTypeIcon")),
-				};
-				check(sizeof(Images) / sizeof(*Images) > (int32)PinContainerType);
-				const FSlateBrush* SecondaryIcon = PinContainerType == EPinContainerType::Map ? FEditorStyle::GetBrush(TEXT("Kismet.VariableList.MapValueTypeIcon")) : nullptr;
-
-				static const FText Tooltips[] = {
-					LOCTEXT("SingleVariableTooltip", "Single Variable"),
-					LOCTEXT("ArrayTooltip", "Array"),
-					LOCTEXT("SetTooltip", "Set"),
-					LOCTEXT("MapTooltip", "Map (Dictionary)"),
-				};
-				check(sizeof(Tooltips) / sizeof(*Tooltips) > (int32)PinContainerType);
-
-
-				return SNew(
-							SDoubleImage,
-							SecondaryIcon,
-							TAttribute<FSlateColor>(this, &SPinTypeSelector::GetSecondaryTypeIconColor)
-						)
-					.Image(Images[(int32)PinContainerType])
-					.ToolTip(IDocumentation::Get()->CreateToolTip(Tooltips[(int32)PinContainerType], nullptr, *BigTooltipDocLink, TEXT("Containers")))
-					.ColorAndOpacity(this, &SPinTypeSelector::GetTypeIconColor);
-			};
-
-			ContainerControl = SNew(SComboBox< TSharedPtr<EPinContainerType> >)
+			ContainerControl = SNew(SComboButton)
 			.ButtonStyle(FCoreStyle::Get(), "NoBorder")
 			.HasDownArrow(false)
-			.OnGenerateWidget(
-				SComboBox< TSharedPtr<EPinContainerType> >::FOnGenerateWidget::CreateLambda(
-					GenerateContainerTypeEntry
+			.OnGetMenuContent(
+				FOnGetContent::CreateLambda(
+					[this]()
+					{
+						typedef SListView< TSharedPtr<EPinContainerType> > SPinContainerListView;
+						return SNew(SPinContainerListView)
+							.ListItemsSource(&PinTypes)
+							.OnGenerateRow(
+								SPinContainerListView::FOnGenerateRow::CreateLambda(
+									[this](TSharedPtr<EPinContainerType> InPinContainerType, const TSharedRef<STableViewBase>& OwnerTable)->TSharedRef<ITableRow>
+									{
+										EPinContainerType PinContainerType = *InPinContainerType;
+										check(sizeof(Images) / sizeof(*Images) > (int32)PinContainerType);
+										check(sizeof(Tooltips) / sizeof(*Tooltips) > (int32)PinContainerType);
+										const FSlateBrush* SecondaryIcon = PinContainerType == EPinContainerType::Map ? FEditorStyle::GetBrush(TEXT("Kismet.VariableList.MapValueTypeIcon")) : nullptr;
+
+										return SNew(STableRow<TSharedPtr<EPinContainerType>>, OwnerTable)
+											.Content()
+											[
+												SNew(
+													SDoubleImage,
+														SecondaryIcon,
+														TAttribute<FSlateColor>(this, &SPinTypeSelector::GetSecondaryTypeIconColor)
+														)
+												.Image(Images[(int32)PinContainerType])
+												.ToolTip(IDocumentation::Get()->CreateToolTip(Tooltips[(int32)PinContainerType], nullptr, *BigTooltipDocLink, TEXT("Containers")))
+												.ColorAndOpacity(this, &SPinTypeSelector::GetTypeIconColor)
+											]
+											.IsEnabled(
+												TAttribute<bool>::Create(TAttribute<bool>::FGetter::CreateLambda( 
+													[PinContainerType, this]() 
+													{ 
+														return !ContainerRequiresGetTypeHash(PinContainerType) || FBlueprintEditorUtils::HasGetTypeHash( this->TargetPinType.Get() );
+													} ) )
+												);
+									}
+								)
+							)
+							.OnSelectionChanged(
+								SPinContainerListView::FOnSelectionChanged::CreateLambda(
+									[this](TSharedPtr<EPinContainerType> InType, ESelectInfo::Type)
+									{
+										this->OnContainerTypeSelectionChanged(*InType);
+									}
+								)
+							);
+					}
 				)
 			)
-			.OptionsSource(&PinTypes)
 			.ContentPadding(0)
 			.ToolTip(IDocumentation::Get()->CreateToolTip( TAttribute<FText>(this, &SPinTypeSelector::GetToolTipForContainerWidget), NULL, *BigTooltipDocLink, TEXT("Containers")))
 			.IsEnabled( TargetPinType.Get().PinCategory != Schema->PC_Exec )
 			.Visibility(InArgs._bAllowArrays ? EVisibility::Visible : EVisibility::Collapsed)
-			.OnSelectionChanged(
-				SComboBox< TSharedPtr<EPinContainerType> >::FOnSelectionChanged::CreateLambda(
-					[this](TSharedPtr<EPinContainerType> InType, ESelectInfo::Type)
-					{
-						this->OnContainerTypeSelectionChanged(*InType);
-					}
-				)
-			)
-			.Content()
+			.ButtonContent()
 			[
 				SNew(
 					SDoubleImage,
@@ -336,11 +361,10 @@ void SPinTypeSelector::Construct(const FArguments& InArgs, FGetPinTypeTree GetPi
 		+SHorizontalBox::Slot()
 		[
 			SNew(SBox)
-			.WidthOverride(110.f)
 			.Visibility(
 				TAttribute<EVisibility>::Create(
 					TAttribute<EVisibility>::FGetter::CreateLambda(
-						[this]() {return this->TargetPinType.Get().bIsMap == true ? EVisibility::Visible : EVisibility::Hidden; }
+						[this]() {return this->TargetPinType.Get().bIsMap == true ? EVisibility::Visible : EVisibility::Collapsed; }
 					)
 				)
 			)
@@ -749,6 +773,17 @@ void SPinTypeSelector::OnSelectPinType(FPinTypeTreeItem InItem, FString InPinCat
 		NewTargetPinType.PinValueType.TerminalSubCategoryObject = nullptr;
 	}
 
+	if ((NewTargetPinType.bIsMap || NewTargetPinType.bIsSet) && !FBlueprintEditorUtils::HasGetTypeHash(NewTargetPinType))
+	{
+		// inform user via toast why the type change was exceptional and clear IsMap/IsSetness because this type cannot be hashed:
+		const FText NotificationText = FText::Format(LOCTEXT("TypeCannotBeHashed", "Container type cleared because '{0}' does not have a GetTypeHash function. Maps and Sets require a hash function to insert and find elements"), GetTypeDescription());
+		FNotificationInfo Info(NotificationText);
+		Info.ExpireDuration = 8.0f;
+		FSlateNotificationManager::Get().AddNotification(Info);
+		NewTargetPinType.bIsMap = false;
+		NewTargetPinType.bIsSet = false;
+	}
+
 	OnTypeChanged.ExecuteIfBound(NewTargetPinType);
 }
 
diff --git a/Engine/Source/Editor/LevelEditor/Private/LevelEditorToolBar.cpp b/Engine/Source/Editor/LevelEditor/Private/LevelEditorToolBar.cpp
--- a/Engine/Source/Editor/LevelEditor/Private/LevelEditorToolBar.cpp
+++ b/Engine/Source/Editor/LevelEditor/Private/LevelEditorToolBar.cpp
@@ -1899,7 +1899,7 @@ TSharedRef< SWidget > FLevelEditorToolBar::GenerateOpenBlueprintMenuContent( TSh
 								FExecuteAction::CreateStatic(&FLevelEditorToolBar::OnOpenSubLevelBlueprint, Level)
 								);
 
-							FText DisplayName = FText::Format(LOCTEXT("SubLevelBlueprintItem", "Edit {LevelName}"), FText::FromString(FPaths::GetCleanFilename(Level->GetOutermost()->GetName())));
+							FText DisplayName = FText::Format(LOCTEXT("SubLevelBlueprintItem", "Edit {0}"), FText::FromString(FPaths::GetCleanFilename(Level->GetOutermost()->GetName())));
 							InMenuBuilder.AddMenuEntry(DisplayName, FText::GetEmpty(), EditBP, UIAction);
 						}
 					}
diff --git a/Engine/Source/Editor/PropertyEditor/Private/UserInterface/PropertyEditor/SPropertyEditorAsset.cpp b/Engine/Source/Editor/PropertyEditor/Private/UserInterface/PropertyEditor/SPropertyEditorAsset.cpp
--- a/Engine/Source/Editor/PropertyEditor/Private/UserInterface/PropertyEditor/SPropertyEditorAsset.cpp
+++ b/Engine/Source/Editor/PropertyEditor/Private/UserInterface/PropertyEditor/SPropertyEditorAsset.cpp
@@ -504,6 +504,10 @@ FText SPropertyEditorAsset::OnGetAssetName() const
 				AActor* Actor = CastChecked<AActor>(Value.Object);
 				Name = FText::FromString(Actor->GetActorLabel());
 			}
+			else if (UField* AsField = Cast<UField>(Value.Object))
+			{
+				Name = AsField->GetDisplayNameText();
+			}
 			else
 			{
 				Name = FText::FromString(Value.Object->GetName());
diff --git a/Engine/Source/Editor/UnrealEd/Private/Dialogs/DlgPickAssetPath.cpp b/Engine/Source/Editor/UnrealEd/Private/Dialogs/DlgPickAssetPath.cpp
--- a/Engine/Source/Editor/UnrealEd/Private/Dialogs/DlgPickAssetPath.cpp
+++ b/Engine/Source/Editor/UnrealEd/Private/Dialogs/DlgPickAssetPath.cpp
@@ -112,12 +112,17 @@ void SDlgPickAssetPath::OnPathChange(const FString& NewPath)
 FReply SDlgPickAssetPath::OnButtonClick(EAppReturnType::Type ButtonID)
 {
 	UserResponse = ButtonID;
-
+	
 	if (ButtonID == EAppReturnType::Cancel || ValidatePackage())
 	{
 		// Only close the window if canceling or if the asset name is valid
 		RequestDestroyWindow();
 	}
+	else
+	{
+		// reset the user response in case the window is closed using 'x'.
+		UserResponse = EAppReturnType::Cancel;
+	}
 
 	return FReply::Handled();
 }
diff --git a/Engine/Source/Editor/UnrealEd/Private/EdGraphUtilities.cpp b/Engine/Source/Editor/UnrealEd/Private/EdGraphUtilities.cpp
--- a/Engine/Source/Editor/UnrealEd/Private/EdGraphUtilities.cpp
+++ b/Engine/Source/Editor/UnrealEd/Private/EdGraphUtilities.cpp
@@ -9,6 +9,7 @@
 #include "SNotificationList.h"
 #include "NotificationManager.h"
 #include "K2Node_TunnelBoundary.h"
+#include "K2Node_Composite.h"
 
 /////////////////////////////////////////////////////
 // FGraphObjectTextFactory
@@ -22,8 +23,6 @@ public:
 	const UEdGraph* DestinationGraph;
 	TArray<FName> ExtraNamesInUse;
 	TArray<UEdGraphNode*> NodesToDestroy;
-
-	TArray<UEdGraphNode*> PendingNodes;
 public:
 	FGraphObjectTextFactory(const UEdGraph* InDestinationGraph)
 		: FCustomizableTextObjectFactory(GWarn)
@@ -59,9 +58,9 @@ protected:
 		return false;
 	}
 
-	void HandleNode(UEdGraphNode* CreatedObject)
+	virtual void ProcessConstructedObject(UObject* CreatedObject) override
 	{
-		if (UEdGraphNode* Node = CreatedObject)
+		if (UEdGraphNode* Node = Cast<UEdGraphNode>(CreatedObject))
 		{
 			UEdGraphNode* CreatedNode = Node;
 			if (!Node->CanPasteHere(DestinationGraph))
@@ -85,48 +84,8 @@ protected:
 		}
 	}
 
-	virtual void ProcessConstructedObject(UObject* CreatedObject) override
-	{
-		if (UEdGraphNode* Node = Cast<UEdGraphNode>(CreatedObject))
-		{
-			PendingNodes.Add(Node);
-		}
-	}
-
 	virtual void PostProcessConstructedObjects() override
 	{
-		// First handle elements that can change the BPGC signature (for example Custom Event nodes).
-		for (int32 Idx = 0; Idx < PendingNodes.Num();)
-		{
-			UK2Node* Node = Cast<UK2Node>(PendingNodes[Idx]);
-			if (Node && Node->NodeCausesStructuralBlueprintChange())
-			{
-				HandleNode(Node);
-				PendingNodes.RemoveAtSwap(Idx);
-			}
-			else
-			{
-				Idx++;
-			}
-		}
-
-		// Update Skel Class signature
-		if (SpawnedNodes.Num())
-		{
-			UBlueprint* Blueprint = FBlueprintEditorUtils::FindBlueprintForGraph(DestinationGraph);
-			if (Blueprint)
-			{
-				FBlueprintEditorUtils::MarkBlueprintAsStructurallyModified(Blueprint);
-			}
-
-		}
-
-		// Handle all remaining nodes. Some of them can depend on the BPGC signature
-		for (UEdGraphNode* Node : PendingNodes)
-		{
-			HandleNode(Node);
-		}
-
 		if (SubstituteNodes.Num() > 0)
 		{
 			// Display a notification to inform the user that the variable type was invalid (likely due to corruption), it should no longer appear in the list.
@@ -139,8 +98,6 @@ protected:
 			{
 				Notification->SetCompletionState(SNotificationItem::CS_None);
 			}
-
-
 		}
 
 		for (UEdGraphNode* Node : NodesToDestroy)
@@ -320,7 +277,41 @@ void FEdGraphUtilities::MergeChildrenGraphsIn(UEdGraph* MergeTarget, UEdGraph* P
 	{
 		UEdGraph* ChildGraph = ParentGraph->SubGraphs[Index];
 
-		if (bWantBoundaryNodes)
+		if (ChildGraph && MessageLog)
+		{
+			UK2Node_Composite* CompositeInstance = ChildGraph->GetTypedOuter<UK2Node_Composite>();
+			if (MessageLog->bTreatCompositeGraphsAsTunnels && CompositeInstance != nullptr)
+			{
+				// Locate or register active tunnels for this child graph.
+				TArray<TWeakObjectPtr<UEdGraphNode>> ActiveTunnels;
+				UEdGraphNode* TrueCompositeInstance = MessageLog->GetSourceTunnelNode(CompositeInstance);
+				MessageLog->GetTunnelsActiveForNode(TrueCompositeInstance, ActiveTunnels);
+				if (!ActiveTunnels.Num())
+				{
+					ActiveTunnels.Add(TrueCompositeInstance);
+					MessageLog->RegisterIntermediateTunnelInstance(TrueCompositeInstance, ActiveTunnels);
+				}
+				// Register composite nodes to the source composite instance.
+				for (auto Node : ChildGraph->Nodes)
+				{
+					MessageLog->RegisterIntermediateTunnelNode(Node, TrueCompositeInstance);
+					if (FBlueprintEditorUtils::IsTunnelInstanceNode(Node))
+					{
+						if (Node->IsA<UK2Node_Composite>())
+						{
+							UEdGraphNode* SourceNode = MessageLog->GetSourceTunnelNode(Node);
+							MessageLog->RegisterIntermediateTunnelInstance(SourceNode, ActiveTunnels);
+						}
+						else
+						{
+							// Register child macro instance nodes using the duplicated instance.
+							MessageLog->RegisterIntermediateTunnelInstance(Node, ActiveTunnels);
+						}
+					}
+				}
+			}
+		}
+		if (bWantBoundaryNodes && MessageLog)
 		{
 			// Create boundary nodes around tunnels for debugging/profiling if requested.
 			UK2Node_TunnelBoundary::CreateBoundaryNodesForGraph(ChildGraph, *MessageLog);
diff --git a/Engine/Source/Editor/UnrealEd/Private/Fbx/ReimportFbxSceneFactory.cpp b/Engine/Source/Editor/UnrealEd/Private/Fbx/ReimportFbxSceneFactory.cpp
--- a/Engine/Source/Editor/UnrealEd/Private/Fbx/ReimportFbxSceneFactory.cpp
+++ b/Engine/Source/Editor/UnrealEd/Private/Fbx/ReimportFbxSceneFactory.cpp
@@ -720,7 +720,7 @@ void UReimportFbxSceneFactory::RecursivelySetComponentProperties(USCS_Node* Curr
 	if (!CurrentNodeActorComponent) //We need a component
 		return;
 
-	int32 IndexTemplateSuffixe = CurrentNodeActorComponent->GetName().Find(TEXT("_GEN_VARIABLE"));
+	int32 IndexTemplateSuffixe = CurrentNodeActorComponent->GetName().Find(USimpleConstructionScript::ComponentTemplateNameSuffix);
 	bool NameContainTemplateSuffixe = IndexTemplateSuffixe != INDEX_NONE;
 	FString NodeName = CurrentNodeActorComponent->GetName();
 	FString ReduceNodeName = NodeName;
@@ -747,7 +747,7 @@ void UReimportFbxSceneFactory::RecursivelySetComponentProperties(USCS_Node* Curr
 
 		if (NameContainTemplateSuffixe)
 		{
-			ComponentName += TEXT("_GEN_VARIABLE");
+			ComponentName += USimpleConstructionScript::ComponentTemplateNameSuffix;
 		}
 		USceneComponent *SceneComponent = Cast<USceneComponent>(ActorComponent);
 		if (!SceneComponent) //We support only scene component
@@ -765,7 +765,7 @@ void UReimportFbxSceneFactory::RecursivelySetComponentProperties(USCS_Node* Curr
 			FString ComponentParentName = ParentComponent->GetName();
 			if (NameContainTemplateSuffixe)
 			{
-				ComponentParentName += TEXT("_GEN_VARIABLE");
+				ComponentParentName += USimpleConstructionScript::ComponentTemplateNameSuffix;
 			}
 			ComponentParentNames.Insert(ComponentParentName, 0);
 			ParentComponent = ParentComponent->GetAttachParent();
diff --git a/Engine/Source/Editor/UnrealEd/Private/Kismet2/BlueprintEditorUtils.cpp b/Engine/Source/Editor/UnrealEd/Private/Kismet2/BlueprintEditorUtils.cpp
--- a/Engine/Source/Editor/UnrealEd/Private/Kismet2/BlueprintEditorUtils.cpp
+++ b/Engine/Source/Editor/UnrealEd/Private/Kismet2/BlueprintEditorUtils.cpp
@@ -2180,6 +2180,16 @@ void FBlueprintEditorUtils::UpdateDelegatesInBlueprint(UBlueprint* Blueprint)
 			{
 				EventNode->UpdateDelegatePin();
 			}
+
+			TArray<UK2Node_Knot*> Knots;
+			Graph->GetNodesOfClass(Knots);
+			for (UK2Node_Knot* Knot : Knots)
+			{
+				// Indiscriminate reuse of UK2Node_Knot::PostReconstructNode() is the convention established
+				// by UEdGraphSchema_K2::OnPinConnectionDoubleCicked. This forces the pin type data to be
+				// refreshed (e.g. due to changes in UpdateDelegatePin())
+				Knot->PostReconstructNode();
+			}
 		}
 	}
 }
@@ -3551,7 +3561,7 @@ int32 FBlueprintEditorUtils::FindSCS_Node(const UBlueprint* Blueprint, const FNa
 	
 		for(int32 i=0; i<AllSCS_Nodes.Num(); i++)
 		{
-			if(AllSCS_Nodes[i]->VariableName == InName)
+			if(AllSCS_Nodes[i]->GetVariableName() == InName)
 			{
 				return i;
 			}
@@ -4366,14 +4376,14 @@ void FBlueprintEditorUtils::RenameComponentMemberVariable(UBlueprint* Blueprint,
 	// Should not allow renaming to "none" (UI should prevent this)
 	check(!NewName.IsNone());
 
-	if (!NewName.IsEqual(Node->VariableName, ENameCase::CaseSensitive))
+	if (!NewName.IsEqual(Node->GetVariableName(), ENameCase::CaseSensitive))
 	{
 		Blueprint->Modify();
 
 		// Update the name
-		const FName OldName = Node->VariableName;
+		const FName OldName = Node->GetVariableName();
 		Node->Modify();
-		Node->VariableName = NewName;
+		Node->SetVariableName(NewName);
 
 		// Rename Inheritable Component Templates
 		{
@@ -4484,7 +4494,7 @@ void FBlueprintEditorUtils::RenameMemberVariable(UBlueprint* Blueprint, const FN
 			for (TArray<USCS_Node*>::TConstIterator NodeIt(Nodes); NodeIt; ++NodeIt)
 			{
 				USCS_Node* CurrentNode = *NodeIt;
-				if (CurrentNode->VariableName == OldName)
+				if (CurrentNode->GetVariableName() == OldName)
 				{
 					RenameComponentMemberVariable(Blueprint, CurrentNode, NewName);
 					break;
@@ -6435,7 +6445,7 @@ bool FBlueprintEditorUtils::IsSCSComponentProperty(UObjectProperty* MemberProper
 					return true;
 				}
 			}
-			else if (ScsNode->VariableName == MemberRef.GetMemberName())
+			else if (ScsNode->GetVariableName() == MemberRef.GetMemberName())
 			{
 				return true;
 			}
@@ -7961,7 +7971,7 @@ bool FBlueprintEditorUtils::IsPaletteActionReadOnly(TSharedPtr<FEdGraphSchemaAct
 			for (TArray<USCS_Node*>::TConstIterator NodeIt(Nodes); NodeIt; ++NodeIt)
 			{
 				USCS_Node* CurrentNode = *NodeIt;
-				if (CurrentNode->VariableName == VarAction->GetVariableName())
+				if (CurrentNode->GetVariableName() == VarAction->GetVariableName())
 				{
 					bIsReadOnly = false;
 					break;
@@ -8396,6 +8406,21 @@ const FSlateBrush* FBlueprintEditorUtils::GetSecondaryIconFromPin(const FEdGraph
 	return nullptr;
 }
 
+bool FBlueprintEditorUtils::HasGetTypeHash(const FEdGraphPinType& PinType)
+{
+	const UEdGraphSchema_K2* K2Schema = GetDefault<UEdGraphSchema_K2>();
+	if (PinType.PinCategory != K2Schema->PC_Struct)
+	{
+		// even object or class types can be hashed, no reason to investigate further
+		return true;
+	}
+
+	// Disabling hasing of structs until work on data integrity is finished. Need to make
+	// sure that we can load sets and maps of sets that have been altered since the container
+	// was originally created:
+	return false;
+}
+
 FText FBlueprintEditorUtils::GetFriendlyClassDisplayName(const UClass* Class)
 {
 	if (Class != nullptr)
diff --git a/Engine/Source/Editor/UnrealEd/Private/Kismet2/CompilerResultsLog.cpp b/Engine/Source/Editor/UnrealEd/Private/Kismet2/CompilerResultsLog.cpp
--- a/Engine/Source/Editor/UnrealEd/Private/Kismet2/CompilerResultsLog.cpp
+++ b/Engine/Source/Editor/UnrealEd/Private/Kismet2/CompilerResultsLog.cpp
@@ -8,6 +8,7 @@
 #include "SourceCodeNavigation.h"
 #include "Developer/HotReload/Public/IHotReload.h"
 #include "EngineLogs.h"
+#include "K2Node_MacroInstance.h"
 
 #if WITH_EDITOR
 
@@ -124,6 +125,7 @@ FCompilerResultsLog::FCompilerResultsLog(bool bIsCompatibleWithEvents/* = true*/
 	, bLogInfoOnly(false)
 	, bAnnotateMentionedNodes(true)
 	, bLogDetailedResults(false)
+	, bTreatCompositeGraphsAsTunnels(false)
 	, EventDisplayThresholdMs(0)
 {
 	CurrentEventScope = nullptr;
@@ -588,6 +590,94 @@ void FCompilerResultsLog::AddChildEvent(TSharedPtr<FCompilerEvent>& ParentEventS
 	}
 }
 
+void FCompilerResultsLog::RegisterIntermediateTunnelNode(UEdGraphNode* IntermediateNode, UEdGraphNode* OwningTunnelInstance)
+{
+	IntermediateTunnelNodeToTunnelInstanceMap.Add(IntermediateNode, OwningTunnelInstance);
+	UObject* TrueSourceNode = SourceBacktrackMap.FindSourceObject(IntermediateNode);
+	if (IntermediateNode != TrueSourceNode)
+	{
+		FinalNodeBackToTunnelSourceMap.NotifyIntermediateObjectCreation(IntermediateNode, TrueSourceNode);
+	}
+}
+
+void FCompilerResultsLog::RegisterIntermediateTunnelInstance(UEdGraphNode* IntermediateTunnel, TArray<TWeakObjectPtr<UEdGraphNode>>& ActiveTunnels)
+{
+	for (auto Tunnel : ActiveTunnels)
+	{
+		if (Tunnel.Get()->IsA<UK2Node_MacroInstance>() || bTreatCompositeGraphsAsTunnels)
+		{
+			IntermediateTunnelInstanceHierarchyMap.Add(IntermediateTunnel, Tunnel);
+		}
+	}
+	UEdGraphNode* TrueSourceTunnelInstance = Cast<UEdGraphNode>(SourceBacktrackMap.FindSourceObject(IntermediateTunnel));
+	if (!ActiveTunnels.Contains(TrueSourceTunnelInstance))
+	{
+		if (TrueSourceTunnelInstance->IsA<UK2Node_MacroInstance>() || bTreatCompositeGraphsAsTunnels)
+		{
+			IntermediateTunnelInstanceHierarchyMap.Add(IntermediateTunnel, TrueSourceTunnelInstance);
+		}
+	}
+}
+
+UEdGraphNode* FCompilerResultsLog::GetSourceNode(const UEdGraphNode* IntermediateNode)
+{
+	UEdGraphNode* Result = nullptr;
+	if (UEdGraphNode* SourceTunnelInstance = GetSourceTunnelInstance(IntermediateNode))
+	{
+		if (SourceTunnelInstance->IsA<UK2Node_MacroInstance>() || bTreatCompositeGraphsAsTunnels)
+		{
+			Result = SourceTunnelInstance;
+		}
+	}
+	if (!Result)
+	{
+		Result = Cast<UEdGraphNode>(SourceBacktrackMap.FindSourceObject(const_cast<UEdGraphNode*>(IntermediateNode)));
+	}
+	return Result;
+}
+
+UEdGraphNode* FCompilerResultsLog::GetIntermediateTunnelInstance(const UEdGraphNode* IntermediateNode)
+{
+	TWeakObjectPtr<UEdGraphNode> Result;
+	if (TWeakObjectPtr<UEdGraphNode>* IntermediateTunnelInstanceNode = IntermediateTunnelNodeToTunnelInstanceMap.Find(IntermediateNode))
+	{
+		Result = *IntermediateTunnelInstanceNode;
+	}
+	return Result.Get();
+}
+
+UEdGraphNode* FCompilerResultsLog::GetSourceTunnelNode(const UEdGraphNode* IntermediateNode)
+{
+	UEdGraphNode* SourceNode = Cast<UEdGraphNode>(FinalNodeBackToTunnelSourceMap.FindSourceObject(const_cast<UEdGraphNode*>(IntermediateNode)));
+	if (SourceNode == IntermediateNode)
+	{
+		// No source node available because we haven't mapped this yet, this could be a top level Tunnel so look it up.
+		SourceNode = Cast<UEdGraphNode>(SourceBacktrackMap.FindSourceObject(const_cast<UEdGraphNode*>(IntermediateNode)));
+	}
+	return SourceNode;
+}
+
+UEdGraphNode* FCompilerResultsLog::GetSourceTunnelInstance(const UEdGraphNode* IntermediateNode)
+{
+	TWeakObjectPtr<UEdGraphNode> Result;
+	if (UEdGraphNode* IntermediateTunnelInstance = GetIntermediateTunnelInstance(IntermediateNode))
+	{
+		if (UEdGraphNode* SourceTunnelInstanceNode = Cast<UEdGraphNode>(SourceBacktrackMap.FindSourceObject(IntermediateTunnelInstance)))
+		{
+			if (SourceTunnelInstanceNode != IntermediateNode)
+			{
+				Result = SourceTunnelInstanceNode;
+			}
+		}
+	}
+	return Result.Get();
+}
+
+void FCompilerResultsLog::GetTunnelsActiveForNode(const UEdGraphNode* IntermediateNode, TArray<TWeakObjectPtr<UEdGraphNode>>& ActiveTunnelsOut)
+{
+	IntermediateTunnelInstanceHierarchyMap.MultiFind(IntermediateNode, ActiveTunnelsOut, true);
+}
+
 #undef LOCTEXT_NAMESPACE
 
-#endif	//#if !WITH_EDITOR
+#endif	//#if WITH_EDITOR
diff --git a/Engine/Source/Editor/UnrealEd/Private/Kismet2/Kismet2.cpp b/Engine/Source/Editor/UnrealEd/Private/Kismet2/Kismet2.cpp
--- a/Engine/Source/Editor/UnrealEd/Private/Kismet2/Kismet2.cpp
+++ b/Engine/Source/Editor/UnrealEd/Private/Kismet2/Kismet2.cpp
@@ -774,6 +774,11 @@ void FKismetEditorUtilities::CompileBlueprint(UBlueprint* BlueprintObj, bool bIs
 	CompileOptions.bSaveIntermediateProducts = bSaveIntermediateProducts;
 	CompileOptions.bRegenerateSkelton = !bSkeletonUpToDate;
 	CompileOptions.bAddInstrumentation = bAddInstrumentation;
+	if (pResults)
+	{
+		// enable debug information for composite graph instances if we are instrumenting the blueprint.
+		pResults->bTreatCompositeGraphsAsTunnels = bAddInstrumentation;
+	}
 	Compiler.CompileBlueprint(BlueprintObj, CompileOptions, Results, ReinstanceHelper);
 
 	FBlueprintEditorUtils::UpdateDelegatesInBlueprint(BlueprintObj);
diff --git a/Engine/Source/Editor/UnrealEd/Private/Kismet2/KismetDebugUtilities.cpp b/Engine/Source/Editor/UnrealEd/Private/Kismet2/KismetDebugUtilities.cpp
--- a/Engine/Source/Editor/UnrealEd/Private/Kismet2/KismetDebugUtilities.cpp
+++ b/Engine/Source/Editor/UnrealEd/Private/Kismet2/KismetDebugUtilities.cpp
@@ -868,7 +868,7 @@ UBreakpoint* FKismetDebugUtilities::FindBreakpointForNode(UBlueprint* Blueprint,
 		{
 			Blueprint->Breakpoints.RemoveAtSwap(Index);
 			Blueprint->MarkPackageDirty();
-			UE_LOG(LogBlueprintDebug, Warning, TEXT("Encountered a blueprint breakpoint in %s without an associated node (this should not happen... if you know how your blueprint got in this state, then please notify the Engine-Blueprints team)"), *Blueprint->GetPathName());
+			UE_LOG(LogBlueprintDebug, Display, TEXT("Encountered a blueprint breakpoint in %s without an associated node. The blueprint has been remoevd"), *Blueprint->GetPathName());
 			continue;
 		}
 
diff --git a/Engine/Source/Editor/UnrealEd/Private/Kismet2/KismetReinstanceUtilities.cpp b/Engine/Source/Editor/UnrealEd/Private/Kismet2/KismetReinstanceUtilities.cpp
--- a/Engine/Source/Editor/UnrealEd/Private/Kismet2/KismetReinstanceUtilities.cpp
+++ b/Engine/Source/Editor/UnrealEd/Private/Kismet2/KismetReinstanceUtilities.cpp
@@ -126,6 +126,66 @@ struct FReplaceReferenceHelper
 	}
 };
 
+struct FArchetypeReinstanceHelper
+{
+	/** Returns the full set of archetypes rooted at a single archetype object, with additional object flags (optional) */
+	static void GetArchetypeObjects(UObject* InObject, TArray<UObject*>& OutArchetypeObjects, EObjectFlags SubArchetypeFlags = RF_NoFlags)
+	{
+		OutArchetypeObjects.Empty();
+
+		if (InObject != nullptr && InObject->HasAllFlags(RF_ArchetypeObject))
+		{
+			OutArchetypeObjects.Add(InObject);
+
+			TArray<UObject*> ArchetypeInstances;
+			InObject->GetArchetypeInstances(ArchetypeInstances);
+
+			for (int32 Idx = 0; Idx < ArchetypeInstances.Num(); ++Idx)
+			{
+				UObject* ArchetypeInstance = ArchetypeInstances[Idx];
+				if (ArchetypeInstance != nullptr && !ArchetypeInstance->IsPendingKill() && ArchetypeInstance->HasAllFlags(RF_ArchetypeObject | SubArchetypeFlags))
+				{
+					OutArchetypeObjects.Add(ArchetypeInstance);
+
+					TArray<UObject*> SubArchetypeInstances;
+					ArchetypeInstance->GetArchetypeInstances(SubArchetypeInstances);
+
+					if (SubArchetypeInstances.Num() > 0)
+					{
+						ArchetypeInstances.Append(SubArchetypeInstances);
+					}
+				}
+			}
+		}
+	}
+
+	/** Returns an object name that's found to be unique within the given set of archetype objects */
+	static FName FindUniqueArchetypeObjectName(TArray<UObject*>& InArchetypeObjects)
+	{
+		FName OutName = NAME_None;
+
+		if (InArchetypeObjects.Num() > 0)
+		{
+			while (OutName == NAME_None)
+			{
+				UObject* ArchetypeObject = InArchetypeObjects[0];
+				OutName = MakeUniqueObjectName(ArchetypeObject->GetOuter(), ArchetypeObject->GetClass());
+				for (int32 ObjIdx = 1; ObjIdx < InArchetypeObjects.Num(); ++ObjIdx)
+				{
+					ArchetypeObject = InArchetypeObjects[ObjIdx];
+					if (StaticFindObjectFast(ArchetypeObject->GetClass(), ArchetypeObject->GetOuter(), OutName))
+					{
+						OutName = NAME_None;
+						break;
+					}
+				}
+			}
+		}
+
+		return OutName;
+	}
+};
+
 /////////////////////////////////////////////////////////////////////////////////
 // FBlueprintCompileReinstancer
 
@@ -1432,6 +1492,9 @@ void FBlueprintCompileReinstancer::ReplaceInstancesOfClass_Inner(TMap<UClass*, U
 	// Map of old objects to new objects
 	TMap<UObject*, UObject*> OldToNewInstanceMap;
 
+	// Map of old objects to new name (used to assist with reinstancing archetypes)
+	TMap<UObject*, FName> OldToNewNameMap;
+
 	TMap<FStringAssetReference, UObject*> ReinstancedObjectsWeakReferenceMap;
 
 	// actors being replace
@@ -1575,22 +1638,17 @@ void FBlueprintCompileReinstancer::ReplaceInstancesOfClass_Inner(TMap<UClass*, U
 						SpawnInfo.bDeferConstruction = true;
 						SpawnInfo.Name = OldActor->GetFName();
 
-						// Temporarily remove the deprecated flag so we can respawn the Blueprint in the level
-						const bool bIsClassDeprecated = SpawnClass->HasAnyClassFlags(CLASS_Deprecated);
-						SpawnClass->ClassFlags &= ~CLASS_Deprecated;
-
 						OldActor->UObject::Rename(nullptr, OldObject->GetOuter(), REN_DoNotDirty | REN_DontCreateRedirectors);
-						AActor* NewActor = World->SpawnActor(SpawnClass, &Location, &Rotation, SpawnInfo);
-						
-						if (OldActor->CurrentTransactionAnnotation.IsValid())
+
+						AActor* NewActor = nullptr;
 						{
-							NewActor->CurrentTransactionAnnotation = OldActor->CurrentTransactionAnnotation;
+							FMakeClassSpawnableOnScope TemporarilySpawnable(SpawnClass);
+							NewActor = World->SpawnActor(SpawnClass, &Location, &Rotation, SpawnInfo);
 						}
 
-						// Reassign the deprecated flag if it was previously assigned
-						if (bIsClassDeprecated)
+						if (OldActor->CurrentTransactionAnnotation.IsValid())
 						{
-							SpawnClass->ClassFlags |= CLASS_Deprecated;
+							NewActor->CurrentTransactionAnnotation = OldActor->CurrentTransactionAnnotation;
 						}
 
 						check(NewActor != nullptr);
@@ -1684,8 +1742,45 @@ void FBlueprintCompileReinstancer::ReplaceInstancesOfClass_Inner(TMap<UClass*, U
 						EObjectFlags OldFlags = OldObject->GetFlags();
 
 						FName OldName(OldObject->GetFName());
-						OldObject->Rename(nullptr, OldObject->GetOuter(), REN_DoNotDirty | REN_DontCreateRedirectors);
-						NewUObject = NewObject<UObject>(OldObject->GetOuter(), NewClass, OldName, RF_NoFlags, NewArchetype);
+
+						// If the old object is in this table, we've already renamed it away in a previous iteration. Don't rename it again!
+						if (!OldToNewNameMap.Contains(OldObject))
+						{
+							// If we're reinstancing a component template, we also need to rename any inherited templates that are found to be based on it, in order to preserve archetype paths.
+							if (bIsComponent && OldObject->HasAllFlags(RF_ArchetypeObject) && OldObject->GetOuter()->IsA<UBlueprintGeneratedClass>())
+							{
+								// Gather all component templates from the current archetype to the farthest antecedent inherited template(s).
+								TArray<UObject*> OldArchetypeObjects;
+								FArchetypeReinstanceHelper::GetArchetypeObjects(OldObject, OldArchetypeObjects, RF_InheritableComponentTemplate);
+
+								// Find a unique object name that does not conflict with anything in the scope of all outers in the template chain.
+								const FString OldArchetypeName = FArchetypeReinstanceHelper::FindUniqueArchetypeObjectName(OldArchetypeObjects).ToString();
+
+								for (UObject* OldArchetypeObject : OldArchetypeObjects)
+								{
+									OldToNewNameMap.Add(OldArchetypeObject, OldName);
+									OldArchetypeObject->Rename(*OldArchetypeName, OldArchetypeObject->GetOuter(), REN_DoNotDirty | REN_DontCreateRedirectors);
+								}
+							}
+							else
+							{
+								OldObject->Rename(nullptr, OldObject->GetOuter(), REN_DoNotDirty | REN_DontCreateRedirectors);
+							}
+						}
+						
+						{
+							// We may have already renamed this object to temp space if it was an inherited archetype in a previous iteration; check for that here.
+							FName NewName = OldToNewNameMap.FindRef(OldObject);
+							if (NewName == NAME_None)
+							{
+								// Otherwise, just use the old object's current name.
+								NewName = OldName;
+							}
+
+							FMakeClassSpawnableOnScope TemporarilySpawnable(NewClass);
+							NewUObject = NewObject<UObject>(OldObject->GetOuter(), NewClass, NewName, RF_NoFlags, NewArchetype);
+						}
+
 						check(NewUObject != nullptr);
 
 						const EObjectFlags FlagMask = RF_Public | RF_ArchetypeObject | RF_Transactional | RF_Transient | RF_TextExportTransient | RF_InheritableComponentTemplate; //TODO: what about RF_RootSet and RF_Standalone ?
@@ -1785,7 +1880,7 @@ void FBlueprintCompileReinstancer::ReplaceInstancesOfClass_Inner(TMap<UClass*, U
 
 					if (bLogConversions)
 					{
-						UE_LOG(LogBlueprint, Log, TEXT("Converted instance '%s' to '%s'"), *OldObject->GetPathName(), *NewUObject->GetPathName());
+						UE_LOG(LogBlueprint, Log, TEXT("Converted instance '%s' to '%s'"), *GetPathNameSafe(OldObject), *GetPathNameSafe(NewUObject));
 					}
 				}
 			}
diff --git a/Engine/Source/Editor/UnrealEd/Public/Kismet2/BlueprintEditorUtils.h b/Engine/Source/Editor/UnrealEd/Public/Kismet2/BlueprintEditorUtils.h
--- a/Engine/Source/Editor/UnrealEd/Public/Kismet2/BlueprintEditorUtils.h
+++ b/Engine/Source/Editor/UnrealEd/Public/Kismet2/BlueprintEditorUtils.h
@@ -1307,6 +1307,11 @@ public:
 	 */
 	static const struct FSlateBrush* GetSecondaryIconFromPin(const FEdGraphPinType& PinType);
 
+	/**
+	 * Returns true if this terminal type can be hashed (native types need GetTypeHash, script types are always hashable).
+	 */
+	static bool HasGetTypeHash(const FEdGraphPinType& PinType);
+
 	/**
 	 * Generate component instancing data (for cooked builds).
 	 *
@@ -1433,4 +1438,37 @@ struct UNREALED_API FBlueprintDuplicationScopeFlags
 
 	TGuardValue<uint32> Guard;
 	FBlueprintDuplicationScopeFlags(uint32 InFlags) : Guard(bStaticFlags, InFlags) {}
+};
+
+struct UNREALED_API FMakeClassSpawnableOnScope
+{
+	UClass* Class;
+	bool bIsDeprecated;
+	bool bIsAbstract;
+	FMakeClassSpawnableOnScope(UClass* InClass)
+		: Class(InClass), bIsDeprecated(false), bIsAbstract(false)
+	{
+		if (Class)
+		{
+			bIsDeprecated = Class->HasAnyClassFlags(CLASS_Deprecated);
+			Class->ClassFlags &= ~CLASS_Deprecated;
+			bIsAbstract = Class->HasAnyClassFlags(CLASS_Abstract);
+			Class->ClassFlags &= ~CLASS_Abstract;
+		}
+	}
+	~FMakeClassSpawnableOnScope()
+	{
+		if (Class)
+		{
+			if (bIsAbstract)
+			{
+				Class->ClassFlags |= CLASS_Abstract;
+			}
+
+			if (bIsDeprecated)
+			{
+				Class->ClassFlags |= CLASS_Deprecated;
+			}
+		}
+	}
 };
\ No newline at end of file
diff --git a/Engine/Source/Editor/UnrealEd/Public/Kismet2/CompilerResultsLog.h b/Engine/Source/Editor/UnrealEd/Public/Kismet2/CompilerResultsLog.h
--- a/Engine/Source/Editor/UnrealEd/Public/Kismet2/CompilerResultsLog.h
+++ b/Engine/Source/Editor/UnrealEd/Public/Kismet2/CompilerResultsLog.h
@@ -90,14 +90,8 @@ public:
 	// Should detailed results be appended to the final summary log?
 	bool bLogDetailedResults;
 
-	// Special maps used for autocreated macros to preserve information about their source
-	FBacktrackMap FinalNodeBackToMacroSourceMap;
-	TMultiMap<TWeakObjectPtr<UEdGraphNode>, TWeakObjectPtr<UEdGraphNode>> MacroSourceToMacroInstanceNodeMap;
-
-	// Used to track node generatation from tunnels, added in addition to existing code avoid causing any fallout.
-	// This will be refactored after blueprint profiler MVP.
-	TMap<TWeakObjectPtr<UEdGraphNode>, TWeakObjectPtr<UEdGraphNode>> SourceNodeToTunnelInstanceNodeMap;
-	TMap<TWeakObjectPtr<UEdGraphNode>, TWeakObjectPtr<UEdGraphNode>> IntermediateTunnelNodeToSourceNodeMap;
+	// Should composite graphs be treated the same as macro graphs/tunnels to generate more accurate debug data.
+	bool bTreatCompositeGraphsAsTunnels;
 
 	// Minimum event time (ms) for inclusion into the final summary log
 	int EventDisplayThresholdMs;
@@ -112,6 +106,21 @@ protected:
 	// Name of the source object being compiled
 	FString SourcePath;
 
+	// Backtrack map for intermendiate tunnel/macro nodes
+	FBacktrackMap FinalNodeBackToTunnelSourceMap;
+	
+	// Map to track intermediate tunnel nodes back to the intermediate expansion tunnel instance.
+	TMap<TWeakObjectPtr<UEdGraphNode>, TWeakObjectPtr<UEdGraphNode>> IntermediateTunnelNodeToTunnelInstanceMap;
+
+	// Map to track active nested tunnels for intermediate tunnel instances.
+	TMultiMap<TWeakObjectPtr<UEdGraphNode>, TWeakObjectPtr<UEdGraphNode>> IntermediateTunnelInstanceHierarchyMap;
+
+	// Map to track compile time only expansion nodes.
+	TSet<TWeakObjectPtr<const UEdGraphNode>> ExpansionNodes;
+
+	// Per pin filtering of trace signals for instrumentation.
+	TSet<const UEdGraphPin*> PinTraceFilter;
+
 public:
 	FCompilerResultsLog(bool bIsCompatibleWithEvents = true);
 	virtual ~FCompilerResultsLog();
@@ -159,6 +168,35 @@ public:
 	void NotifyIntermediateObjectCreation(UObject* NewObject, UObject* SourceObject);
 	void NotifyIntermediatePinCreation(UEdGraphPin* NewObject, UEdGraphPin* SourceObject);
 
+	/** Compile time node expansion tracking and testing */
+	void RegisterExpansionNode(UEdGraphNode* ExpansionNode) { ExpansionNodes.Add(ExpansionNode); }
+	bool IsExpansionNode(const UEdGraphNode* Node) const { return ExpansionNodes.Contains(Node); }
+
+	/** Per pin trace signal filtering */
+	void AddPinTraceFilter(UEdGraphPin* PinToFilter) { PinTraceFilter.Add(PinToFilter); }
+	bool IsPinTraceSuppressed(const UEdGraphPin* PotentiallyFilteredPin) const { return PinTraceFilter.Contains(PotentiallyFilteredPin); }
+
+	/** Registers intermediate tunnel nodes, both the node and tunnel instance should be intermediate nodes */
+	void RegisterIntermediateTunnelNode(UEdGraphNode* Node, UEdGraphNode* OwningTunnelInstance);
+
+	/** Registers an intermediate tunnel instance node, the active tunnels should not be intermediate nodes. */
+	void RegisterIntermediateTunnelInstance(UEdGraphNode* IntermediateTunnel, TArray<TWeakObjectPtr<UEdGraphNode>>& ActiveTunnels);
+
+	/** Returns the source tunnel instance or the source tunnel node depending on what created the intermediate node */
+	UEdGraphNode* GetSourceNode(const UEdGraphNode* IntermediateNode);
+
+	/** Returns the intermediate tunnel instance that generated the node */
+	UEdGraphNode* GetIntermediateTunnelInstance(const UEdGraphNode* IntermediateNode);
+
+	/** Returns the source tunnel node for the intermediate node */
+	UEdGraphNode* GetSourceTunnelNode(const UEdGraphNode* IntermediateNode);
+
+	/** Returns the source tunnel instance that generated the intermediate node */
+	UEdGraphNode* GetSourceTunnelInstance(const UEdGraphNode* IntermediateNode);
+
+	/** Returns the active tunnel instances for the intermediate node */
+	void GetTunnelsActiveForNode(const UEdGraphNode* IntermediateNode, TArray<TWeakObjectPtr<UEdGraphNode>>& ActiveTunnelsOut);
+
 	/** Returns the true source object for the passed in object */
 	UObject* FindSourceObject(UObject* PossiblyDuplicatedObject);
 	UObject const* FindSourceObject(UObject const* PossiblyDuplicatedObject) const;
diff --git a/Engine/Source/Programs/UnrealHeaderTool/Private/HeaderParser.cpp b/Engine/Source/Programs/UnrealHeaderTool/Private/HeaderParser.cpp
--- a/Engine/Source/Programs/UnrealHeaderTool/Private/HeaderParser.cpp
+++ b/Engine/Source/Programs/UnrealHeaderTool/Private/HeaderParser.cpp
@@ -887,6 +887,15 @@ namespace
 			// Script VM doesn't support array of weak ptrs.
 			return IsPropertySupportedByBlueprint(ArrayProperty->Inner, false);
 		}
+		else if (const USetProperty* SetProperty = Cast<const USetProperty>(Property))
+		{
+			return IsPropertySupportedByBlueprint(SetProperty->ElementProp, false);
+		}
+		else if (const UMapProperty* MapProperty = Cast<const UMapProperty>(Property))
+		{
+			return IsPropertySupportedByBlueprint(MapProperty->KeyProp, false) &&
+				IsPropertySupportedByBlueprint(MapProperty->ValueProp, false);
+		}
 
 		const bool bSupportedType = Property->IsA<UInterfaceProperty>()
 			|| Property->IsA<UClassProperty>()
diff --git a/Engine/Source/Runtime/Core/Public/Misc/DateTime.h b/Engine/Source/Runtime/Core/Public/Misc/DateTime.h
--- a/Engine/Source/Runtime/Core/Public/Misc/DateTime.h
+++ b/Engine/Source/Runtime/Core/Public/Misc/DateTime.h
@@ -694,7 +694,7 @@ public:
 	 * @param DateTime The date and time to get the hash for.
 	 * @return Hash value.
 	 */
-	friend uint32 GetTypeHash(const FDateTime& DateTime);
+	friend CORE_API uint32 GetTypeHash(const FDateTime& DateTime);
 
 protected:
 
diff --git a/Engine/Source/Runtime/Core/Public/Misc/Timespan.h b/Engine/Source/Runtime/Core/Public/Misc/Timespan.h
--- a/Engine/Source/Runtime/Core/Public/Misc/Timespan.h
+++ b/Engine/Source/Runtime/Core/Public/Misc/Timespan.h
@@ -609,7 +609,7 @@ public:
 	 * @param Timespan The timespan to get the hash for.
 	 * @return Hash value.
 	 */
-	friend uint32 GetTypeHash(const FTimespan& Timespan);
+	friend CORE_API uint32 GetTypeHash(const FTimespan& Timespan);
 
 protected:
 
diff --git a/Engine/Source/Runtime/Core/Public/Templates/HasGetTypeHash.h b/Engine/Source/Runtime/Core/Public/Templates/HasGetTypeHash.h
new file mode 100644
--- /dev/null
+++ b/Engine/Source/Runtime/Core/Public/Templates/HasGetTypeHash.h
@@ -0,0 +1,41 @@
+// Copyright 1998-2016 Epic Games, Inc. All Rights Reserved.
+
+#pragma once
+
+namespace UE4GetTypeHashExists_Private
+{
+	struct FNotSpecified {};
+
+	template <typename T>
+	struct FReturnValueCheck
+	{
+		static char (&Func())[2];
+	};
+
+	template <>
+	struct FReturnValueCheck<FNotSpecified>
+	{
+		static char (&Func())[1];
+	};
+
+	template <typename T>
+	FNotSpecified GetTypeHash(const T&);
+
+	template <typename T>
+	const T& Make();
+
+	template <typename T>
+	struct GetTypeHashQuery
+	{
+		enum { Value = sizeof(FReturnValueCheck<decltype(GetTypeHash(Make<T>()))>::Func()) == sizeof(char[2]) };
+	};
+}
+
+/**
+* Traits class which tests if a type has a GetTypeHash overload.
+*/
+template <typename T>
+struct THasGetTypeHash
+{
+	enum { Value = UE4GetTypeHashExists_Private::GetTypeHashQuery<T>::Value };
+};
diff --git a/Engine/Source/Runtime/Core/Public/UObject/BlueprintsObjectVersion.h b/Engine/Source/Runtime/Core/Public/UObject/BlueprintsObjectVersion.h
--- a/Engine/Source/Runtime/Core/Public/UObject/BlueprintsObjectVersion.h
+++ b/Engine/Source/Runtime/Core/Public/UObject/BlueprintsObjectVersion.h
@@ -14,6 +14,7 @@ struct CORE_API FBlueprintsObjectVersion
 		EdGraphPinOptimized,
 		AllowDeletionConformed,
 		AdvancedContainerSupport,
+		SCSHasComponentTemplateClass,
 
 		// -----<new versions can be added above this line>-------------------------------------------------
 		VersionPlusOne,
diff --git a/Engine/Source/Runtime/CoreUObject/Private/Blueprint/BlueprintSupport.cpp b/Engine/Source/Runtime/CoreUObject/Private/Blueprint/BlueprintSupport.cpp
--- a/Engine/Source/Runtime/CoreUObject/Private/Blueprint/BlueprintSupport.cpp
+++ b/Engine/Source/Runtime/CoreUObject/Private/Blueprint/BlueprintSupport.cpp
@@ -1525,7 +1525,7 @@ void FLinkerLoad::ResolveDeferredExports(UClass* LoadClass)
 		// sub-classes of this Blueprint could have had their CDO's 
 		// initialization deferred (this occurs when the sub-class CDO is 
 		// created before this super CDO has been fully serialized; we do this
-		// because the  sub-class's CDO would not have been initialized with 
+		// because the sub-class's CDO would not have been initialized with 
 		// accurate values)
 		//
 		// in that case, the sub-class CDOs are waiting around until their 
@@ -2008,6 +2008,24 @@ void FDeferredObjInitializerTracker::ResolveDeferredSubObjects(UObject* CDO)
 
 	for (FObjectInitializer* DeferredInitializer : DeferredSubObjInitializers)
 	{
+		UObject* SubObjArchetype = DeferredInitializer->GetArchetype();
+		// just because the class's CDO archetype has completed serialization 
+		// does not mean that the archetypes for these sub-objects have... 
+		// unfortunately there isn't anywhere else where we ensure that these 
+		// archetypes are loaded (else we should defer this further, until then 
+		// when we can guarantee that each archetype is complete) - we don't  
+		// even force a load of these sub-object archetypes prior to their own 
+		// Blueprint's regeneration, meaning it's compiled-on-load potentially 
+		// with incomplete component templates; this is potentially bad in 
+		// itself (TODO: investigate)
+		if (SubObjArchetype && !SubObjArchetype->HasAnyFlags(RF_LoadCompleted))
+		{
+			if (FLinkerLoad* SubObjLinker = SubObjArchetype->GetLinker())
+			{
+				SubObjArchetype->SetFlags(RF_NeedLoad);
+				SubObjLinker->Preload(SubObjArchetype);
+			}
+		}
 		FScriptIntegrationObjectHelper::PostConstructInitObject(*DeferredInitializer);
 	}
 	ThreadInst.DeferredSubObjInitializers.Remove(LoadClass);
diff --git a/Engine/Source/Runtime/CoreUObject/Private/UObject/Class.cpp b/Engine/Source/Runtime/CoreUObject/Private/UObject/Class.cpp
--- a/Engine/Source/Runtime/CoreUObject/Private/UObject/Class.cpp
+++ b/Engine/Source/Runtime/CoreUObject/Private/UObject/Class.cpp
@@ -1312,18 +1312,6 @@ void UStruct::AddReferencedObjects(UObject* InThis, FReferenceCollector& Collect
 		// Required by the unified GC when running in the editor
 		Collector.AddReferencedObject( This->SuperStruct, This );
 		Collector.AddReferencedObject( This->Children, This );
-
-		TArray<UObject*> ScriptObjectReferences;
-		FArchiveScriptReferenceCollector ObjectReferenceCollector( ScriptObjectReferences );
-		int32 iCode = 0;
-		while( iCode < This->Script.Num() )
-		{	
-			This->SerializeExpr( iCode, ObjectReferenceCollector );
-		}
-		for( int32 Index = 0; Index < ScriptObjectReferences.Num(); Index++ )
-		{
-			Collector.AddReferencedObject( ScriptObjectReferences[ Index ], This );
-		}
 	}
 
 	//@todo NickW, temp hack to make stale property chains less crashy
@@ -1390,6 +1378,19 @@ bool UStruct::GetStringMetaDataHierarchical(const FName& Key, FString* OutValue)
 	return false;
 }
 
+const UStruct* UStruct::HasMetaDataHierarchical(const FName& Key) const
+{
+	for (const UStruct* TestStruct = this; TestStruct != nullptr; TestStruct = TestStruct->GetSuperStruct())
+	{
+		if (TestStruct->HasMetaData(Key))
+		{
+			return TestStruct;
+		}
+	}
+
+	return nullptr;
+}
+
 #endif
 
 #if USE_CIRCULAR_DEPENDENCY_LOAD_DEFERRING
@@ -3909,7 +3910,7 @@ UFunction* UClass::FindFunctionByName(FName InName, EIncludeSuperFlag::Type Incl
 			{
 				for (const FImplementedInterface& Inter : Interfaces)
 				{
-					Result = Inter.Class->FindFunctionByName(InName);
+					Result = Inter.Class ? Inter.Class->FindFunctionByName(InName) : nullptr;
 					if (Result)
 					{
 						break;
diff --git a/Engine/Source/Runtime/CoreUObject/Private/UObject/LinkerLoad.cpp b/Engine/Source/Runtime/CoreUObject/Private/UObject/LinkerLoad.cpp
--- a/Engine/Source/Runtime/CoreUObject/Private/UObject/LinkerLoad.cpp
+++ b/Engine/Source/Runtime/CoreUObject/Private/UObject/LinkerLoad.cpp
@@ -4266,6 +4266,14 @@ UObject* FLinkerLoad::CreateImport( int32 Index )
 	
 						// Find object now that we know it's class, outer and name.
 						FindObject = FindImportFast(FindClass, FindOuter, Import.ObjectName);
+						if (UDynamicClass* FoundDynamicClass = Cast<UDynamicClass>(FindObject))
+						{
+							if (!FoundDynamicClass->GetDefaultObject(false))
+							{
+								// This class wasn't fully constructed yet. It will be properly constructed in CreateExport. 
+								FindObject = nullptr;
+							}
+						}
 					}
 
 					if( FindObject )
diff --git a/Engine/Source/Runtime/CoreUObject/Private/UObject/ScriptCore.cpp b/Engine/Source/Runtime/CoreUObject/Private/UObject/ScriptCore.cpp
--- a/Engine/Source/Runtime/CoreUObject/Private/UObject/ScriptCore.cpp
+++ b/Engine/Source/Runtime/CoreUObject/Private/UObject/ScriptCore.cpp
@@ -1621,7 +1621,7 @@ void UObject::execInstrumentation( FFrame& Stack, RESULT_DECL )
 	if (EventType == EScriptInstrumentation::InlineEvent)
 	{
 		const FName& EventName = *reinterpret_cast<FName*>(&Stack.Code[1]);
-		FScriptInstrumentationSignal InstrumentationEventInfo(EScriptInstrumentation::Event, this, Stack, EventName);
+		FScriptInstrumentationSignal InstrumentationEventInfo(EventType, this, Stack, EventName);
 		FBlueprintCoreDelegates::InstrumentScriptEvent(InstrumentationEventInfo);
 		Stack.SkipCode(sizeof(FName) + 1);
 	}
diff --git a/Engine/Source/Runtime/CoreUObject/Private/UObject/StructScriptLoader.cpp b/Engine/Source/Runtime/CoreUObject/Private/UObject/StructScriptLoader.cpp
--- a/Engine/Source/Runtime/CoreUObject/Private/UObject/StructScriptLoader.cpp
+++ b/Engine/Source/Runtime/CoreUObject/Private/UObject/StructScriptLoader.cpp
@@ -2,6 +2,7 @@
 
 #include "CoreUObjectPrivate.h"
 #include "StructScriptLoader.h"
+#include "ArchiveScriptReferenceCollector.h"
 
 /*******************************************************************************
  * FDeferredScriptTracker
diff --git a/Engine/Source/Runtime/CoreUObject/Private/UObject/StructScriptLoader.h b/Engine/Source/Runtime/CoreUObject/Private/UObject/StructScriptLoader.h
--- a/Engine/Source/Runtime/CoreUObject/Private/UObject/StructScriptLoader.h
+++ b/Engine/Source/Runtime/CoreUObject/Private/UObject/StructScriptLoader.h
@@ -87,44 +87,3 @@ private:
 	/** Denotes where to offset the archiver to start reading in bytecode */
 	int32 ScriptSerializationOffset;
 }; 
-
-/*******************************************************************************
- * FArchiveScriptReferenceCollector
- ******************************************************************************/
-
-class FArchiveScriptReferenceCollector : public FArchiveUObject
-{
-public:
-	/**
-	* Constructor
-	*
-	* @param	InObjectArray			Array to add object references to
-	*/
-	FArchiveScriptReferenceCollector(TArray<UObject*>& InObjectArray)
-		: ObjectArray(InObjectArray)
-	{
-		ArIsObjectReferenceCollector = true;
-		ArIsPersistent = false;
-		ArIgnoreArchetypeRef = false;
-	}
-protected:
-	/**
-	* UObject serialize operator implementation
-	*
-	* @param Object	reference to Object reference
-	* @return reference to instance of this class
-	*/
-	FArchive& operator<<(UObject*& Object)
-	{
-		// Avoid duplicate entries.
-		if (Object != NULL && !ObjectArray.Contains(Object))
-		{
-			check(Object->IsValidLowLevel());
-			ObjectArray.Add(Object);
-		}
-		return *this;
-	}
-
-	/** Stored reference to array of objects we add object references to */
-	TArray<UObject*>&		ObjectArray;
-};
diff --git a/Engine/Source/Runtime/CoreUObject/Private/UObject/UObjectGlobals.cpp b/Engine/Source/Runtime/CoreUObject/Private/UObject/UObjectGlobals.cpp
--- a/Engine/Source/Runtime/CoreUObject/Private/UObject/UObjectGlobals.cpp
+++ b/Engine/Source/Runtime/CoreUObject/Private/UObject/UObjectGlobals.cpp
@@ -2333,13 +2333,27 @@ UObject* StaticAllocateObject
 		}
 		else if(!InOuter || !InOuter->HasAnyFlags(RF_ClassDefaultObject))
 		{
+#if !UE_BUILD_SHIPPING
+			// Handle nested DSOs
+			bool bIsOwnedByCDO = false;
+			UObject* Iter = InOuter;
+			while (Iter)
+			{
+				if (Iter->HasAnyFlags(RF_ClassDefaultObject))
+				{
+					bIsOwnedByCDO = true;
+					break;
+				}
+				Iter = Iter->GetOuter();
+			}
 			// Should only get in here if we're NOT creating a subobject of a CDO.  CDO subobjects may still need to be serialized off of disk after being created by the constructor
 			// if really necessary there was code to allow replacement of object just needing postload, but lets not go there unless we have to
-			checkf(!Obj->HasAnyFlags(RF_NeedLoad|RF_NeedPostLoad|RF_ClassDefaultObject),
+			checkf(!Obj->HasAnyFlags(RF_NeedLoad|RF_NeedPostLoad|RF_ClassDefaultObject) || bIsOwnedByCDO,
 				*FText::Format(NSLOCTEXT("Core", "ReplaceNotFullyLoaded_f", "Attempting to replace an object that hasn't been fully loaded: {0} (Outer={1}, Flags={2})"),
 					FText::FromString(Obj->GetFullName()),
 					InOuter ? FText::FromString(InOuter->GetFullName()) : FText::FromString(TEXT("NULL")),
 					FText::FromString(FString::Printf(TEXT("0x%08x"), (int32)Obj->GetFlags()))).ToString());
+#endif//UE_BUILD_SHIPPING
 		}
 		// Subobjects are always created in the constructor, no need to re-create them here unless their archetype != CDO or they're blueprint generated.	
 		if (!bCreatingCDO && (!bCanRecycleSubobjects || !Obj->IsDefaultSubobject()))
diff --git a/Engine/Source/Runtime/CoreUObject/Public/Serialization/ArchiveScriptReferenceCollector.h b/Engine/Source/Runtime/CoreUObject/Public/Serialization/ArchiveScriptReferenceCollector.h
new file mode 100644
--- /dev/null
+++ b/Engine/Source/Runtime/CoreUObject/Public/Serialization/ArchiveScriptReferenceCollector.h
@@ -0,0 +1,46 @@
+// Copyright 1998-2016 Epic Games, Inc. All Rights Reserved.
+
+#pragma once
+
+#include "ArchiveUObject.h"
+
+/*******************************************************************************
+ * FArchiveScriptReferenceCollector
+ ******************************************************************************/
+
+class FArchiveScriptReferenceCollector : public FArchiveUObject
+{
+public:
+	/**
+	* Constructor
+	*
+	* @param	InObjectArray			Array to add object references to
+	*/
+	FArchiveScriptReferenceCollector(TArray<UObject*>& InObjectArray)
+		: ObjectArray(InObjectArray)
+	{
+		ArIsObjectReferenceCollector = true;
+		ArIsPersistent = false;
+		ArIgnoreArchetypeRef = false;
+	}
+protected:
+	/**
+	* UObject serialize operator implementation
+	*
+	* @param Object	reference to Object reference
+	* @return reference to instance of this class
+	*/
+	FArchive& operator<<(UObject*& Object)
+	{
+		// Avoid duplicate entries.
+		if (Object != NULL && !ObjectArray.Contains(Object))
+		{
+			check(Object->IsValidLowLevel());
+			ObjectArray.Add(Object);
+		}
+		return *this;
+	}
+
+	/** Stored reference to array of objects we add object references to */
+	TArray<UObject*>&		ObjectArray;
+};
diff --git a/Engine/Source/Runtime/CoreUObject/Public/UObject/Class.h b/Engine/Source/Runtime/CoreUObject/Public/UObject/Class.h
--- a/Engine/Source/Runtime/CoreUObject/Public/UObject/Class.h
+++ b/Engine/Source/Runtime/CoreUObject/Public/UObject/Class.h
@@ -425,6 +425,14 @@ public:
 
 	/** Try and find string metadata with the given key. If not found on this class, work up hierarchy looking for it. */
 	bool GetStringMetaDataHierarchical(const FName& Key, FString* OutValue = nullptr) const;
+
+	/**
+	* Determines if the struct or any of its super structs has any metadata associated with the provided key
+	*
+	* @param Key The key to lookup in the metadata
+	* @return pointer to the UStruct that has associated metadata, nullptr if Key is not associated with any UStruct in the hierarchy
+	*/
+	const UStruct* HasMetaDataHierarchical(const FName& Key) const;
 #endif
 
 #if HACK_HEADER_GENERATOR
diff --git a/Engine/Source/Runtime/Engine/Classes/Components/ActorComponent.h b/Engine/Source/Runtime/Engine/Classes/Components/ActorComponent.h
--- a/Engine/Source/Runtime/Engine/Classes/Components/ActorComponent.h
+++ b/Engine/Source/Runtime/Engine/Classes/Components/ActorComponent.h
@@ -823,11 +823,6 @@ private:
 	virtual void Tick( float DeltaTime ) final { check(0); }
 
 #endif
-
-public:
-
-	/** Prefix used to identify template component instances */
-	static const FString ComponentTemplateNameSuffix;
 };
 
 //////////////////////////////////////////////////////////////////////////
diff --git a/Engine/Source/Runtime/Engine/Classes/EdGraph/EdGraphPin.h b/Engine/Source/Runtime/Engine/Classes/EdGraph/EdGraphPin.h
--- a/Engine/Source/Runtime/Engine/Classes/EdGraph/EdGraphPin.h
+++ b/Engine/Source/Runtime/Engine/Classes/EdGraph/EdGraphPin.h
@@ -528,7 +528,6 @@ private:
 	static bool SerializePin(FArchive& Ar, UEdGraphPin*& PinRef, int32 ArrayIdx, UEdGraphPin* RequestingPin, EPinResolveType ResolveType, TArray<UEdGraphPin*>& OldPins);
 	static FString ExportText_PinReference(const UEdGraphPin* Pin);
 	static FString ExportText_PinArray(const TArray<UEdGraphPin*>& PinArray);
-	static bool ImportText_PinReference(const TCHAR*& Buffer, UEdGraphPin*& PinRef, int32 ArrayIdx, UEdGraphPin* RequestingPin, EPinResolveType ResolveType);
 	static bool ImportText_PinArray(const TCHAR*& Buffer, TArray<UEdGraphPin*>& ArrayRef, UEdGraphPin* RequestingPin, EPinResolveType ResolveType);
 };
 
diff --git a/Engine/Source/Runtime/Engine/Classes/Engine/BlueprintGeneratedClass.h b/Engine/Source/Runtime/Engine/Classes/Engine/BlueprintGeneratedClass.h
--- a/Engine/Source/Runtime/Engine/Classes/Engine/BlueprintGeneratedClass.h
+++ b/Engine/Source/Runtime/Engine/Classes/Engine/BlueprintGeneratedClass.h
@@ -317,6 +317,30 @@ public:
 		return INDEX_NONE;
 	}
 
+	// Finds all code locations (Function+CodeOffset) associated with the source node
+	void FindAllCodeLocationsFromSourceNode(UEdGraphNode* SourceNode, UFunction* InFunction, TArray<int32>& OutNodeToCodeAssociations) const
+	{
+		OutNodeToCodeAssociations.Empty();
+
+		if (const FDebuggingInfoForSingleFunction* pFuncInfo = PerFunctionLineNumbers.Find(InFunction))
+		{
+			for (auto CodeLocation : pFuncInfo->LineNumberToSourceNodeMap)
+			{
+				if (CodeLocation.Value == SourceNode)
+				{
+					OutNodeToCodeAssociations.Add(CodeLocation.Key);
+				}
+			}
+			for (auto CodeLocation : pFuncInfo->LineNumberToMacroSourceNodeMap)
+			{
+				if (CodeLocation.Value == SourceNode)
+				{
+					OutNodeToCodeAssociations.Add(CodeLocation.Key);
+				}
+			}
+		}
+	}
+
 	// Finds the pure node script code range associated with the [impure] source node, or FInt32Range(INDEX_NONE) if there is no existing association
 	FInt32Range FindPureNodeScriptCodeRangeFromSourceNode(const UEdGraphNode* SourceNode, UFunction* InFunction) const
 	{
diff --git a/Engine/Source/Runtime/Engine/Classes/Engine/InheritableComponentHandler.h b/Engine/Source/Runtime/Engine/Classes/Engine/InheritableComponentHandler.h
--- a/Engine/Source/Runtime/Engine/Classes/Engine/InheritableComponentHandler.h
+++ b/Engine/Source/Runtime/Engine/Classes/Engine/InheritableComponentHandler.h
@@ -65,6 +65,9 @@ struct FComponentOverrideRecord
 {
 	GENERATED_USTRUCT_BODY()
 
+	UPROPERTY()
+	UClass* ComponentClass;
+
 	UPROPERTY()
 	UActorComponent* ComponentTemplate;
 
@@ -75,7 +78,8 @@ struct FComponentOverrideRecord
 	FBlueprintCookedComponentInstancingData CookedComponentInstancingData;
 
 	FComponentOverrideRecord()
-		: ComponentTemplate(nullptr)
+		: ComponentClass(nullptr)
+		, ComponentTemplate(nullptr)
 	{}
 };
 
diff --git a/Engine/Source/Runtime/Engine/Classes/Engine/ObjectLibrary.h b/Engine/Source/Runtime/Engine/Classes/Engine/ObjectLibrary.h
--- a/Engine/Source/Runtime/Engine/Classes/Engine/ObjectLibrary.h
+++ b/Engine/Source/Runtime/Engine/Classes/Engine/ObjectLibrary.h
@@ -43,6 +43,7 @@ public:
 	//~ Begin UObject Interface
 #if WITH_EDITOR
 	virtual void PostEditChangeProperty(FPropertyChangedEvent& PropertyChangedEvent) override;
+	virtual void PostLoad() override;
 #endif // WITH_EDITOR
 	//~ End UObject Interface
 
diff --git a/Engine/Source/Runtime/Engine/Classes/Engine/SCS_Node.h b/Engine/Source/Runtime/Engine/Classes/Engine/SCS_Node.h
--- a/Engine/Source/Runtime/Engine/Classes/Engine/SCS_Node.h
+++ b/Engine/Source/Runtime/Engine/Classes/Engine/SCS_Node.h
@@ -12,6 +12,10 @@ class USCS_Node : public UObject
 {
 	GENERATED_UCLASS_BODY()
 
+	/** Component class */
+	UPROPERTY()
+	UClass* ComponentClass;
+
 	/** Template for the component to create */
 	UPROPERTY()
 	class UActorComponent* ComponentTemplate;
@@ -21,7 +25,9 @@ class USCS_Node : public UObject
 	FBlueprintCookedComponentInstancingData CookedComponentInstancingData;
 
 	/** If non-None, creates a variable in the class and assigns the created blueprint to it */
-	UPROPERTY()
+	// Note: Keeping this in place for backwards-compatibility with existing projects/tools that might be accessing it.
+	DEPRECATED(4.14, "VariableName is no longer directly accessible; please use GetVariableName()/SetVariableName() instead.")
+	UPROPERTY(transient)
 	FName VariableName;
 
 #if WITH_EDITORONLY_DATA
@@ -75,6 +81,7 @@ class USCS_Node : public UObject
 	bool bVariableNameAutoGenerated_DEPRECATED;
 
 	//~ Begin UObject Interface
+	virtual void Serialize(FArchive& Ar) override;
 	virtual void PostLoad() override;
 	//~ End UObject Interface
 
@@ -131,12 +138,18 @@ class USCS_Node : public UObject
 	/** Preloads the node, and all its child nodes recursively */
 	ENGINE_API void PreloadChain();
 
-	/** Get name of the variable we should create for this component instance */
-	ENGINE_API FName GetVariableName() const;
-
 	/** See if this node is the root */
 	ENGINE_API bool IsRootNode() const;
 
+	/** Get name of the variable we should create for this component instance */
+	FORCEINLINE FName GetVariableName() const
+	{
+		return InternalVariableName;
+	}
+
+	/** Set name of the variable we should create for this component instance */
+	ENGINE_API void SetVariableName(const FName& NewName, bool bRenameTemplate = true);
+
 	/** Get the SimpleConstructionScript that owns this node */
 	class USimpleConstructionScript* GetSCS() const
 	{
@@ -171,8 +184,20 @@ class USCS_Node : public UObject
 	/** Finds and returns the parent component template through the given Blueprint */
 	ENGINE_API USceneComponent* GetParentComponentTemplate(UBlueprint* InBlueprint) const;
 #endif
+
+protected:
+	/** Helper method to rename the given component template along with any instances */
+	static void RenameComponentTemplate(UActorComponent* ComponentTemplate, const FName& NewName);
 		
 private:
+	/** Internal variable name. This is used for:
+			a) Generating the component template (archetype) object name.
+			b) A UObjectProperty in the generated Blueprint class. This holds a reference to the component instance created at Actor construction time.
+			c) Archetype lookup through the generated Blueprint class. All instances route back to the archetype through the variable name (i.e. not the template name).
+	 */
+	UPROPERTY()
+	FName InternalVariableName;
+
 	/** Delegate to trigger when Variable name is modified outside of SCS */
 	FSCSNodeNameChanged		  OnNameChangedExternal;
 };
diff --git a/Engine/Source/Runtime/Engine/Classes/Engine/SimpleConstructionScript.h b/Engine/Source/Runtime/Engine/Classes/Engine/SimpleConstructionScript.h
--- a/Engine/Source/Runtime/Engine/Classes/Engine/SimpleConstructionScript.h
+++ b/Engine/Source/Runtime/Engine/Classes/Engine/SimpleConstructionScript.h
@@ -11,6 +11,9 @@ class USimpleConstructionScript : public UObject
 {
 	GENERATED_UCLASS_BODY()
 
+	/** Suffix used for component template object name. */
+	ENGINE_API static const FString ComponentTemplateNameSuffix;
+
 	//~ Begin UObject Interface
 	virtual void Serialize(FArchive& Ar) override;
 	virtual void PostLoad() override;
@@ -21,14 +24,18 @@ class USimpleConstructionScript : public UObject
 	/** Ensures that all root node parent references are still valid and clears the reference if not */
 	ENGINE_API void FixupRootNodeParentReferences();
 
+	/** Helper method to register instanced components post-construction */
+	static void RegisterInstancedComponent(UActorComponent* Component);
+
 	/**
 	 * Execute this script on the supplied actor, creating components 
 	 *
 	 * @param Actor					The actor instance to execute the script on.
+	 * @param NativeSceneComponents	The set of native scene components from the actor that SCS nodes can parent to when instanced.
 	 * @param RootTransform			The transform to apply to the root scene component of the actor instance if defined in this script.
 	 * @param bIsDefaultTransform	Indicates whether or not the given transform is a "default" transform, in which case it can be overridden by template defaults.
 	 */
-	void ExecuteScriptOnActor(AActor* Actor, const FTransform& RootTransform, bool bIsDefaultTransform);
+	void ExecuteScriptOnActor(AActor* Actor, const TInlineComponentArray<USceneComponent*>& NativeSceneComponents, const FTransform& RootTransform, bool bIsDefaultTransform);
 
 	/** Create the map from names to SCS_Nodes to improve FindSCSNode performance during construction script execution */
 	void CreateNameToSCSNodeMap();
diff --git a/Engine/Source/Runtime/Engine/Classes/GameFramework/Actor.h b/Engine/Source/Runtime/Engine/Classes/GameFramework/Actor.h
--- a/Engine/Source/Runtime/Engine/Classes/GameFramework/Actor.h
+++ b/Engine/Source/Runtime/Engine/Classes/GameFramework/Actor.h
@@ -192,6 +192,9 @@ private:
 	/** Whether we've tried to register tick functions. Reset when they are unregistered. */
 	uint8 bTickFunctionsRegistered : 1;
 
+	/** Whether we've deferred the RegisterAllComponents() call at spawn time. Reset when RegisterAllComponents() is called. */
+	uint8 bHasDeferredComponentRegistration : 1;
+
 	/**
 	 * Enables any collision on this actor.
 	 * @see SetActorEnableCollision(), GetActorEnableCollision()
@@ -2200,6 +2203,9 @@ public:
 	/** Called after all the components in the Components array are registered */
 	virtual void PostRegisterAllComponents();
 
+	/** Returns true if Actor has deferred the RegisterAllComponents() call at spawn time (e.g. pending Blueprint SCS execution to set up a scene root component). */
+	FORCEINLINE bool HasDeferredComponentRegistration() const { return bHasDeferredComponentRegistration; }
+
 	/** Returns true if Actor has a registered root component */
 	bool HasValidRootComponent();
 
@@ -2421,14 +2427,14 @@ protected:
 	*/
 	bool CheckActorComponents();
 
-	/** Checks for and resolve any name conflicts prior to instancing a new Blueprint Component. */
-	void CheckComponentInstanceName(const FName InName);
-
 	/** Called after instancing a new Blueprint Component from either a template or cooked data. */
 	void PostCreateBlueprintComponent(UActorComponent* NewActorComp);
 
 public:
 
+	/** Checks for and resolve any name conflicts prior to instancing a new Blueprint Component. */
+	void CheckComponentInstanceName(const FName InName);
+
 	/** Walk up the attachment chain from RootComponent until we encounter a different actor, and return it. If we are not attached to a component in a different actor, returns NULL */
 	virtual AActor* GetAttachParentActor() const;
 
diff --git a/Engine/Source/Runtime/Engine/Classes/Kismet/BlueprintSetLibrary.h b/Engine/Source/Runtime/Engine/Classes/Kismet/BlueprintSetLibrary.h
new file mode 100644
--- /dev/null
+++ b/Engine/Source/Runtime/Engine/Classes/Kismet/BlueprintSetLibrary.h
@@ -0,0 +1,473 @@
+// Copyright 1998-2016 Epic Games, Inc. All Rights Reserved.
+
+#pragma once
+
+#include "BlueprintSetLibrary.generated.h"
+
+UCLASS(Experimental)
+class ENGINE_API UBlueprintSetLibrary : public UBlueprintFunctionLibrary
+{
+	GENERATED_UCLASS_BODY()
+
+	/** 
+	 *Add item to set
+	 *
+	 *@param	TargetSet		The set to add item to
+	 *@param	NewItem			The item to add to the set
+	*/
+	UFUNCTION(BlueprintCallable, CustomThunk, meta=(DisplayName = "Add", CompactNodeTitle = "ADD", SetParam = "TargetSet|NewItem", AutoCreateRefTerm = "NewItem"), Category="Utilities|Set")
+	static void Set_Add(const TSet<int32>& TargetSet, const int32& NewItem);
+	
+	/**
+	 * Adds all elements from an Array to a Set
+	 *
+	 *@param	TargetSet		The set to search for the item
+	 *@param	NewItems		The items to add to the set
+	*/
+	UFUNCTION(BlueprintCallable, CustomThunk, meta=(DisplayName = "Add Items", CompactNodeTitle = "ADD ITEMS", SetParam = "TargetSet|NewItems", AutoCreateRefTerm = "NewItems"), Category="Utilities|Set")
+	static void Set_AddItems(const TSet<int32>& TargetSet, const TArray<int32>& NewItems);
+
+	/**
+	 *Remove item from set.
+	 *
+	 *@param	TargetSet		The set to remove from
+	 *@param	Item			The item to remove from the set
+	 *@return	True if an item was removed
+	*/
+	UFUNCTION(BlueprintCallable, CustomThunk, meta=(DisplayName = "Remove", CompactNodeTitle = "REMOVE", SetParam = "TargetSet|Item", AutoCreateRefTerm = "Item"), Category="Utilities|Set")
+	static bool Set_Remove(const TSet<int32>& TargetSet, const int32& Item);
+	
+	/**
+	 *Removes all elements in an Array from a set.
+	 *
+	 *@param	TargetSet		The set to remove from
+	 *@param	Items			The items to remove from the set
+	*/
+	UFUNCTION(BlueprintCallable, CustomThunk, meta=(DisplayName = "Remove Items", CompactNodeTitle = "REMOVE ITEMS", SetParam = "TargetSet|Items", AutoCreateRefTerm = "Items"), Category="Utilities|Set")
+	static void Set_RemoveItems(const TSet<int32>& TargetSet, const  TArray<int32>& Items);
+
+	/**
+	* Returns an Array, Result, with the entries of a Set.
+	*
+	* @param		A		Set
+	* @param		Result	Array
+	*/
+	UFUNCTION(BlueprintCallable, CustomThunk, meta = (DisplayName = "To Array", CompactNodeTitle = "TO ARRAY", SetParam = "A|Result"), Category = "Utilities|Set")
+	static void Set_ToArray(const TSet<int32>& A, TArray<int32>& Result);
+
+	/**
+	 *Clear a set, removes all content
+	 *
+	 *@param	TargetSet		The set to clear
+	*/
+	UFUNCTION(BlueprintCallable, CustomThunk, meta=(DisplayName = "Clear", CompactNodeTitle = "CLEAR", Keywords = "empty", SetParam = "TargetSet"), Category="Utilities|Set")
+	static void Set_Clear(const TSet<int32>& TargetSet);
+
+	/**
+	 * Get the number of items in a set
+	 *
+	 *@param	TargetSet		The set to get the length of
+	 *@return	The length of the set
+	*/
+	UFUNCTION(BlueprintPure, CustomThunk, meta=(DisplayName = "Length", CompactNodeTitle = "LENGTH", SetParam = "TargetSet", Keywords = "num size count"), Category="Utilities|Set")
+	static int32 Set_Length(const TSet<int32>& TargetSet);
+
+	/**
+	 *Returns true if the set contains the given item
+	 *
+	 *@param	TargetSet		The set to search for the item
+	 *@param	ItemToFind		The item to look for
+	 *@return	True if the item was found within the set
+	*/
+	UFUNCTION(BlueprintPure, CustomThunk, meta=(DisplayName = "Contains Item", CompactNodeTitle = "CONTAINS", SetParam = "TargetSet|ItemToFind", AutoCreateRefTerm = "ItemToFind"), Category="Utilities|Set")
+	static bool Set_Contains(const TSet<int32>& TargetSet, const int32& ItemToFind);
+
+	/**
+	 * Assigns Result to the intersection of Set A and Set B. That is, Result will contain
+	 * all elements that are in both Set A and Set B. To intersect with the empty set use
+	 * Clear.
+	 *
+	 *@param		A		One set to intersect
+	 *@param		B		Another set to intersect
+	 *@param		Result	Set to store results in
+	 */
+	UFUNCTION(BlueprintCallable, CustomThunk, meta=(DisplayName = "Intersection", CompactNodeTitle = "INTERSECTION", SetParam = "A,B,Result"), Category="Utilities|Set")
+	static void Set_Intersection(const TSet<int32>& A, const TSet<int32>& B, TSet<int32>& Result );
+
+	/**
+	 * Assigns Result to the union of two sets, A and B. That is, Result will contain
+	 * all elements that are in Set A and in addition all elements in Set B. Note that 
+	 * a Set is a collection of unique elements, so duplicates will be eliminated.
+	 *
+	 *@param		A		One set to union
+	 *@param		B		Another set to union
+	 *@param		Result	Set to store results in
+	 */
+	UFUNCTION(BlueprintCallable, CustomThunk, meta = (DisplayName = "Union", CompactNodeTitle = "UNION", SetParam = "A,B,Result"), Category = "Utilities|Set")
+	static void Set_Union(const TSet<int32>& A, const TSet<int32>& B, TSet<int32>& Result );
+
+	/**
+	 * Assigns Result to the relative difference of two sets, A and B. That is, Result will 
+	 * contain  all elements that are in Set A but are not found in Set B. Note that the 
+	 * difference between two sets  is not commutative. The Set whose elements you wish to 
+	 * preserve should be the first (top) parameter. Also called the relative complement.
+	 *
+	 *@param		A		Starting set
+	 *@param		B		Set of elements to remove from set A
+	 *@param		Result	Set containing all elements in A that are not found in B
+	 */
+	UFUNCTION(BlueprintCallable, CustomThunk, meta = (DisplayName = "Difference", CompactNodeTitle = "DIFFERENCE", SetParam = "A,B,Result"), Category = "Utilities|Set")
+	static void Set_Difference(const TSet<int32>& A, const TSet<int32>& B, TSet<int32>& Result );
+
+	DECLARE_FUNCTION(execSet_Add)
+	{
+		Stack.MostRecentProperty = nullptr;
+		Stack.StepCompiledIn<USetProperty>(NULL);
+		void* SetAddr = Stack.MostRecentPropertyAddress;
+		USetProperty* SetProperty = Cast<USetProperty>(Stack.MostRecentProperty);
+		if (!SetProperty)
+		{
+			Stack.bArrayContextFailed = true;
+			return;
+		}
+
+		// Since ItemPtr isn't really an int, step the stack manually
+		const UProperty* ElementProp = SetProperty->ElementProp;
+		const int32 PropertySize = ElementProp->ElementSize * ElementProp->ArrayDim;
+		void* StorageSpace = FMemory_Alloca(PropertySize);
+		ElementProp->InitializeValue(StorageSpace);
+
+		Stack.MostRecentPropertyAddress = NULL;
+		Stack.StepCompiledIn<UProperty>(StorageSpace);
+		void* ItemPtr = StorageSpace;
+
+		P_FINISH;
+
+		P_NATIVE_BEGIN;
+		GenericSet_Add(SetAddr, SetProperty, ItemPtr);
+		P_NATIVE_END;
+
+		ElementProp->DestroyValue(StorageSpace);
+	}
+
+	DECLARE_FUNCTION(execSet_AddItems)
+	{
+		Stack.MostRecentProperty = nullptr;
+		Stack.StepCompiledIn<USetProperty>(NULL);
+		void* SetAddr = Stack.MostRecentPropertyAddress;
+		USetProperty* SetProperty = Cast<USetProperty>(Stack.MostRecentProperty);
+		if (!SetProperty)
+		{
+			Stack.bArrayContextFailed = true;
+			return;
+		}
+
+		Stack.MostRecentProperty = nullptr;
+		Stack.StepCompiledIn<UArrayProperty>(NULL);
+		void* TargetArrayAddr = Stack.MostRecentPropertyAddress;
+		UArrayProperty* TargetArrayProperty = Cast<UArrayProperty>(Stack.MostRecentProperty);
+		if (!TargetArrayProperty)
+		{
+			Stack.bArrayContextFailed = true;
+			return;
+		}
+
+		P_FINISH;
+
+		P_NATIVE_BEGIN;
+		GenericSet_AddItems(SetAddr, SetProperty, TargetArrayAddr, TargetArrayProperty);
+		P_NATIVE_END;
+	}
+
+	DECLARE_FUNCTION(execSet_Remove)
+	{
+		Stack.MostRecentProperty = nullptr;
+		Stack.StepCompiledIn<USetProperty>(NULL);
+		void* SetAddr = Stack.MostRecentPropertyAddress;
+		USetProperty* SetProperty = Cast<USetProperty>(Stack.MostRecentProperty);
+		if (!SetProperty)
+		{
+			Stack.bArrayContextFailed = true;
+			return;
+		}
+
+		// Since ItemPtr isn't really an int, step the stack manually
+		const UProperty* ElementProp = SetProperty->ElementProp;
+		const int32 PropertySize = ElementProp->ElementSize * ElementProp->ArrayDim;
+		void* StorageSpace = FMemory_Alloca(PropertySize);
+		ElementProp->InitializeValue(StorageSpace);
+
+		Stack.MostRecentPropertyAddress = NULL;
+		Stack.StepCompiledIn<UProperty>(StorageSpace);
+		void* ItemPtr = StorageSpace;
+
+		P_FINISH;
+
+		P_NATIVE_BEGIN;
+		GenericSet_Remove(SetAddr, SetProperty, ItemPtr);
+		P_NATIVE_END;
+
+		ElementProp->DestroyValue(StorageSpace);
+	}
+
+	DECLARE_FUNCTION(execSet_RemoveItems)
+	{
+		Stack.MostRecentProperty = nullptr;
+		Stack.StepCompiledIn<USetProperty>(NULL);
+		void* SetAddr = Stack.MostRecentPropertyAddress;
+		USetProperty* SetProperty = Cast<USetProperty>(Stack.MostRecentProperty);
+		if (!SetProperty)
+		{
+			Stack.bArrayContextFailed = true;
+			return;
+		}
+
+		Stack.MostRecentProperty = nullptr;
+		Stack.StepCompiledIn<UArrayProperty>(NULL);
+		void* TargetArrayAddr = Stack.MostRecentPropertyAddress;
+		UArrayProperty* TargetArrayProperty = Cast<UArrayProperty>(Stack.MostRecentProperty);
+		if (!TargetArrayProperty)
+		{
+			Stack.bArrayContextFailed = true;
+			return;
+		}
+
+		P_FINISH;
+
+		P_NATIVE_BEGIN;
+		GenericSet_RemoveItems(SetAddr, SetProperty, TargetArrayAddr, TargetArrayProperty);
+		P_NATIVE_END;
+	}
+
+	DECLARE_FUNCTION(execSet_ToArray)
+	{
+		Stack.MostRecentProperty = nullptr;
+		Stack.StepCompiledIn<USetProperty>(NULL);
+		void* SetAddr = Stack.MostRecentPropertyAddress;
+		USetProperty* SetProperty = Cast<USetProperty>(Stack.MostRecentProperty);
+		if (!SetProperty)
+		{
+			Stack.bArrayContextFailed = true;
+			return;
+		}
+
+		Stack.MostRecentProperty = nullptr;
+		Stack.StepCompiledIn<UArrayProperty>(NULL);
+		void* TargetArrayAddr = Stack.MostRecentPropertyAddress;
+		UArrayProperty* TargetArrayProperty = Cast<UArrayProperty>(Stack.MostRecentProperty);
+		if (!TargetArrayProperty)
+		{
+			Stack.bArrayContextFailed = true;
+			return;
+		}
+
+		P_FINISH;
+
+		P_NATIVE_BEGIN;
+		GenericSet_ToArray(SetAddr, SetProperty, TargetArrayAddr, TargetArrayProperty);
+		P_NATIVE_END;
+	}
+
+	DECLARE_FUNCTION(execSet_Clear)
+	{
+		Stack.MostRecentProperty = nullptr;
+		Stack.StepCompiledIn<USetProperty>(NULL);
+		void* SetAddr = Stack.MostRecentPropertyAddress;
+		USetProperty* SetProperty = Cast<USetProperty>(Stack.MostRecentProperty);
+		if (!SetProperty)
+		{
+			Stack.bArrayContextFailed = true;
+			return;
+		}
+
+		P_FINISH;
+
+		// Perform the search
+		P_NATIVE_BEGIN;
+		GenericSet_Clear(SetAddr, SetProperty);
+		P_NATIVE_END;
+
+	}
+
+	DECLARE_FUNCTION(execSet_Length)
+	{
+		Stack.MostRecentProperty = nullptr;
+		Stack.StepCompiledIn<USetProperty>(NULL);
+		void* SetAddr = Stack.MostRecentPropertyAddress;
+		USetProperty* SetProperty = Cast<USetProperty>(Stack.MostRecentProperty);
+		if (!SetProperty)
+		{
+			Stack.bArrayContextFailed = true;
+			return;
+		}
+
+		P_FINISH;
+
+		// Perform the search
+		P_NATIVE_BEGIN;
+		*(int32*)RESULT_PARAM = GenericSet_Length(SetAddr, SetProperty);
+		P_NATIVE_END;
+	}
+
+	DECLARE_FUNCTION(execSet_Contains)
+	{
+		Stack.MostRecentProperty = nullptr;
+		Stack.StepCompiledIn<USetProperty>(NULL);
+		void* SetAddr = Stack.MostRecentPropertyAddress;
+		USetProperty* SetProperty = Cast<USetProperty>(Stack.MostRecentProperty);
+		if (!SetProperty)
+		{
+			// @todo(dano): rename to 'container context failed'
+			Stack.bArrayContextFailed = true;
+			return;
+		}
+
+		// Since ItemToFind isn't really an int, step the stack manually
+		const UProperty* ElementProp = SetProperty->ElementProp;
+		const int32 PropertySize = ElementProp->ElementSize * ElementProp->ArrayDim;
+		void* StorageSpace = FMemory_Alloca(PropertySize);
+		ElementProp->InitializeValue(StorageSpace);
+
+		Stack.MostRecentPropertyAddress = NULL;
+		Stack.StepCompiledIn<UProperty>(StorageSpace);
+		void* ItemToFindPtr = StorageSpace;
+
+		P_FINISH;
+
+		P_NATIVE_BEGIN;
+		*(bool*)RESULT_PARAM = GenericSet_Contains(SetAddr, SetProperty, ItemToFindPtr);
+		P_NATIVE_END;
+
+		ElementProp->DestroyValue(StorageSpace);
+	}
+
+	DECLARE_FUNCTION(execSet_Intersection)
+	{
+		Stack.MostRecentProperty = nullptr;
+		Stack.StepCompiledIn<USetProperty>(NULL);
+		void* SetAddrA = Stack.MostRecentPropertyAddress;
+		USetProperty* SetPropertyA = Cast<USetProperty>(Stack.MostRecentProperty);
+		if (!SetPropertyA)
+		{
+			Stack.bArrayContextFailed = true;
+			return;
+		}
+
+		Stack.MostRecentProperty = nullptr;
+		Stack.StepCompiledIn<USetProperty>(NULL);
+		void* SetAddrB = Stack.MostRecentPropertyAddress;
+		USetProperty* SetPropertyB = Cast<USetProperty>(Stack.MostRecentProperty);
+		if (!SetPropertyB)
+		{
+			Stack.bArrayContextFailed = true;
+			return;
+		}
+
+		Stack.MostRecentProperty = nullptr;
+		Stack.StepCompiledIn<USetProperty>(NULL);
+		void* SetAddrResult = Stack.MostRecentPropertyAddress;
+		USetProperty* SetPropertyResult = Cast<USetProperty>(Stack.MostRecentProperty);
+		if (!SetPropertyResult)
+		{
+			Stack.bArrayContextFailed = true;
+			return;
+		}
+
+		P_FINISH;
+
+		P_NATIVE_BEGIN;
+		GenericSet_Intersect(SetAddrA, SetPropertyA, SetAddrB, SetPropertyB, SetAddrResult, SetPropertyResult);
+		P_NATIVE_END;
+	}
+
+	DECLARE_FUNCTION(execSet_Union)
+	{
+		Stack.MostRecentProperty = nullptr;
+		Stack.StepCompiledIn<USetProperty>(NULL);
+		void* SetAddrA = Stack.MostRecentPropertyAddress;
+		USetProperty* SetPropertyA = Cast<USetProperty>(Stack.MostRecentProperty);
+		if (!SetPropertyA)
+		{
+			Stack.bArrayContextFailed = true;
+			return;
+		}
+
+		Stack.MostRecentProperty = nullptr;
+		Stack.StepCompiledIn<USetProperty>(NULL);
+		void* SetAddrB = Stack.MostRecentPropertyAddress;
+		USetProperty* SetPropertyB = Cast<USetProperty>(Stack.MostRecentProperty);
+		if (!SetPropertyB)
+		{
+			Stack.bArrayContextFailed = true;
+			return;
+		}
+
+		Stack.MostRecentProperty = nullptr;
+		Stack.StepCompiledIn<USetProperty>(NULL);
+		void* SetAddrResult = Stack.MostRecentPropertyAddress;
+		USetProperty* SetPropertyResult = Cast<USetProperty>(Stack.MostRecentProperty);
+		if (!SetPropertyResult)
+		{
+			Stack.bArrayContextFailed = true;
+			return;
+		}
+
+		P_FINISH;
+
+		P_NATIVE_BEGIN;
+		GenericSet_Union(SetAddrA, SetPropertyA, SetAddrB, SetPropertyB, SetAddrResult, SetPropertyResult);
+		P_NATIVE_END;
+	}
+
+	DECLARE_FUNCTION(execSet_Difference)
+	{
+		Stack.MostRecentProperty = nullptr;
+		Stack.StepCompiledIn<USetProperty>(NULL);
+		void* SetAddrA = Stack.MostRecentPropertyAddress;
+		USetProperty* SetPropertyA = Cast<USetProperty>(Stack.MostRecentProperty);
+		if (!SetPropertyA)
+		{
+			Stack.bArrayContextFailed = true;
+			return;
+		}
+
+		Stack.MostRecentProperty = nullptr;
+		Stack.StepCompiledIn<USetProperty>(NULL);
+		void* SetAddrB = Stack.MostRecentPropertyAddress;
+		USetProperty* SetPropertyB = Cast<USetProperty>(Stack.MostRecentProperty);
+		if (!SetPropertyB)
+		{
+			Stack.bArrayContextFailed = true;
+			return;
+		}
+
+		Stack.MostRecentProperty = nullptr;
+		Stack.StepCompiledIn<USetProperty>(NULL);
+		void* SetAddrResult = Stack.MostRecentPropertyAddress;
+		USetProperty* SetPropertyResult = Cast<USetProperty>(Stack.MostRecentProperty);
+		if (!SetPropertyResult)
+		{
+			Stack.bArrayContextFailed = true;
+			return;
+		}
+
+		P_FINISH;
+
+		P_NATIVE_BEGIN;
+		GenericSet_Difference(SetAddrA, SetPropertyA, SetAddrB, SetPropertyB, SetAddrResult, SetPropertyResult);
+		P_NATIVE_END;
+	}
+
+	static void GenericSet_Add(const void* TargetSet, const USetProperty* SetProperty, const void* ItemPtr);
+	static void GenericSet_AddItems(const void* TargetSet, const USetProperty* SetProperty, const void* TargetArray, const UArrayProperty* ArrayProperty);
+	static void GenericSet_Remove(const void* TargetSet, const USetProperty* SetProperty, const void* ItemPtr);
+	static void GenericSet_RemoveItems(const void* TargetSet, const USetProperty* SetProperty, const void* TargetArray, const UArrayProperty* ArrayProperty);
+	static void GenericSet_ToArray(const void* TargetSet, const USetProperty* SetProperty, void* TargetArray, const UArrayProperty* ArrayProperty);
+	static void GenericSet_Clear(const void* TargetSet, const USetProperty* SetProperty);
+	static int32 GenericSet_Length(const void* TargetSet, const USetProperty* SetProperty);
+	static bool GenericSet_Contains(const void* TargetSet, const USetProperty* SetProperty, const void* ItemToFind);
+	static void GenericSet_Intersect(const void* SetA, const USetProperty* SetPropertyA, const void* SetB, const USetProperty* SetPropertyB, const void* SetResult, const USetProperty* SetPropertyResult);
+	static void GenericSet_Union(const void* SetA, const USetProperty* SetPropertyA, const void* SetB, const USetProperty* SetPropertyB, const void* SetResult, const USetProperty* SetPropertyResult);
+	static void GenericSet_Difference(const void* SetA, const USetProperty* SetPropertyA, const void* SetB, const USetProperty* SetPropertyB, const void* SetResult, const USetProperty* SetPropertyResult);
+};
+
diff --git a/Engine/Source/Runtime/Engine/Classes/Kismet/KismetMathLibrary.h b/Engine/Source/Runtime/Engine/Classes/Kismet/KismetMathLibrary.h
--- a/Engine/Source/Runtime/Engine/Classes/Kismet/KismetMathLibrary.h
+++ b/Engine/Source/Runtime/Engine/Classes/Kismet/KismetMathLibrary.h
@@ -781,11 +781,11 @@ class ENGINE_API UKismetMathLibrary : public UBlueprintFunctionLibrary
 	static float VSize2DSquared(FVector2D A);
 
 	/* Returns a unit normal version of the FVector A */
-	UFUNCTION(BlueprintPure, meta=(DisplayName = "Normalize"), Category="Math|Vector")
+	UFUNCTION(BlueprintPure, meta=(DisplayName = "Normalize", Keywords="Unit Vector"), Category="Math|Vector")
 	static FVector Normal(FVector A);
 
 	/* Returns a unit normal version of the vector2d A */
-	UFUNCTION(BlueprintPure, meta=(DisplayName = "Normalize2D"), Category="Math|Vector2D")
+	UFUNCTION(BlueprintPure, meta=(DisplayName = "Normalize2D", Keywords="Unit Vector"), Category="Math|Vector2D")
 	static FVector2D Normal2D(FVector2D A);
 
 	/* Linearly interpolates between A and B based on Alpha (100% of A when Alpha=0 and 100% of B when Alpha=1) */
@@ -935,8 +935,8 @@ class ENGINE_API UKismetMathLibrary : public UBlueprintFunctionLibrary
 	static FVector GetVectorArrayAverage(const TArray<FVector>& Vectors);
 
 	/** Find the unit direction vector from one position to another. */
-	UFUNCTION(BlueprintPure, Category="Math|Vector")
-	static FVector GetDirectionVector(FVector From, FVector To);
+	UFUNCTION(BlueprintPure, meta=(DisplayName = "Get Unit Direction Vector", Keywords = "Unit Vector"), Category="Math|Vector")
+	static FVector GetDirectionUnitVector(FVector From, FVector To);
 
 	//
 	// Rotator functions.
diff --git a/Engine/Source/Runtime/Engine/Private/Actor.cpp b/Engine/Source/Runtime/Engine/Private/Actor.cpp
--- a/Engine/Source/Runtime/Engine/Private/Actor.cpp
+++ b/Engine/Source/Runtime/Engine/Private/Actor.cpp
@@ -98,6 +98,7 @@ void AActor::InitializeDefaults()
 	bAllowReceiveTickEventOnDedicatedServer = true;
 	bRelevantForNetworkReplays = true;
 	bGenerateOverlapEventsDuringLevelStreaming = false;
+	bHasDeferredComponentRegistration = false;
 #if WITH_EDITORONLY_DATA
 	PivotOffset = FVector::ZeroVector;
 #endif
@@ -2699,7 +2700,7 @@ void AActor::PostSpawnInitialize(FTransform const& UserSpawnTransform, AActor* I
 	ExchangeNetRoles(bRemoteOwned);
 
 	USceneComponent* const SceneRootComponent = FixupNativeActorComponents(this);
-	if (SceneRootComponent != NULL)
+	if (SceneRootComponent != nullptr)
 	{
 		// Set the actor's location and rotation since it has a native rootcomponent
 		// Note that we respect any initial transformation the root component may have from the CDO, so the final transform
@@ -2712,8 +2713,14 @@ void AActor::PostSpawnInitialize(FTransform const& UserSpawnTransform, AActor* I
 	// Call OnComponentCreated on all default (native) components
 	DispatchOnComponentsCreated(this);
 
-	// Initialize the actor's components.
-	RegisterAllComponents();
+	// Register the actor's default (native) components, but only if we have a native scene root. If we don't, it implies that there could be only non-scene components
+	// at the native class level. In that case, if this is a Blueprint instance, we need to defer native registration until after SCS execution can establish a scene root.
+	// Note: This API will also call PostRegisterAllComponents() on the actor instance. If deferred, PostRegisterAllComponents() won't be called until the root is set by SCS.
+	bHasDeferredComponentRegistration = (SceneRootComponent == nullptr && Cast<UBlueprintGeneratedClass>(GetClass()) != nullptr);
+	if (!bHasDeferredComponentRegistration)
+	{
+		RegisterAllComponents();
+	}
 
 	// Set owner.
 	SetOwner(InOwner);
@@ -3929,6 +3936,9 @@ void AActor::RegisterAllComponents()
 {
 	// 0 - means register all components
 	verify(IncrementalRegisterComponents(0));
+
+	// Clear this flag as it's no longer deferred
+	bHasDeferredComponentRegistration = false;
 }
 
 // Walks through components hierarchy and returns closest to root parent component that is unregistered
diff --git a/Engine/Source/Runtime/Engine/Private/ActorConstruction.cpp b/Engine/Source/Runtime/Engine/Private/ActorConstruction.cpp
--- a/Engine/Source/Runtime/Engine/Private/ActorConstruction.cpp
+++ b/Engine/Source/Runtime/Engine/Private/ActorConstruction.cpp
@@ -653,6 +653,31 @@ bool AActor::ExecuteConstruction(const FTransform& Transform, const FComponentIn
 	{
 		if (bErrorFree)
 		{
+			// Get all scene components owned by the given actor prior to SCS execution
+			// Note: GetComponents() internally does a NULL check, so we can assume here that all entries are valid.
+			TInlineComponentArray<USceneComponent*> AllSceneComponents;
+			GetComponents(AllSceneComponents);
+
+			// Determine the set of native scene components that SCS nodes can attach to.
+			TInlineComponentArray<USceneComponent*> NativeSceneComponents;
+			NativeSceneComponents.Reserve(AllSceneComponents.Num());
+			for (USceneComponent* SceneComponent : AllSceneComponents)
+			{
+				// Exclude subcomponents of native components, as these could unintentionally be matched by name during SCS execution. Also exclude instance-only components.
+				if (SceneComponent->CreationMethod == EComponentCreationMethod::Native && SceneComponent->GetOuter()->IsA<AActor>())
+				{
+					// If RootComponent is not set, the first unattached native scene component will be used as root. This matches what's done in FixupNativeActorComponents().
+					// @TODO - consider removing this; keeping here as a fallback just in case it wasn't set prior to SCS execution, but in most cases now this should be valid. 
+					if (RootComponent == nullptr && SceneComponent->GetAttachParent() == nullptr)
+					{
+						// Note: All native scene components should already have been registered at this point, so we don't need to register the component here.
+						SetRootComponent(SceneComponent);
+					}
+
+					NativeSceneComponents.Add(SceneComponent);
+				}
+			}
+
 			// Prevent user from spawning actors in User Construction Script
 			TGuardValue<bool> AutoRestoreISCS(GetWorld()->bIsRunningConstructionScript, true);
 			for (int32 i = ParentBPClassStack.Num() - 1; i >= 0; i--)
@@ -663,12 +688,30 @@ bool AActor::ExecuteConstruction(const FTransform& Transform, const FComponentIn
 				if (SCS)
 				{
 					SCS->CreateNameToSCSNodeMap();
-					SCS->ExecuteScriptOnActor(this, Transform, bIsDefaultTransform);
+					SCS->ExecuteScriptOnActor(this, NativeSceneComponents, Transform, bIsDefaultTransform);
 				}
 				// Now that the construction scripts have been run, we can create timelines and hook them up
 				UBlueprintGeneratedClass::CreateComponentsForActor(CurrentBPGClass, this);
 			}
 
+			// Ensure that we've called RegisterAllComponents(), in case it was deferred and the SCS could not be fully executed.
+			if (HasDeferredComponentRegistration())
+			{
+				RegisterAllComponents();
+			}
+
+			// Once SCS execution has finished, we do a final pass to register any components that may have been deferred or were otherwise left unregistered after SCS execution.
+			TInlineComponentArray<UActorComponent*> PostSCSComponents;
+			GetComponents(PostSCSComponents);
+			for (UActorComponent* ActorComponent : PostSCSComponents)
+			{
+				// Limit registration to components that are known to have been created during SCS execution
+				if (!ActorComponent->IsRegistered() && (ActorComponent->CreationMethod == EComponentCreationMethod::SimpleConstructionScript || !AllSceneComponents.Contains(ActorComponent)))
+				{
+					USimpleConstructionScript::RegisterInstancedComponent(ActorComponent);
+				}
+			}
+
 			// If we passed in cached data, we apply it now, so that the UserConstructionScript can use the updated values
 			if (InstanceDataCache)
 			{
@@ -736,6 +779,12 @@ bool AActor::ExecuteConstruction(const FTransform& Transform, const FComponentIn
 				SetRootComponent(BillboardComponent);
 				FinishAndRegisterComponent(BillboardComponent);
 			}
+
+			// Ensure that we've called RegisterAllComponents(), in case it was deferred and the SCS could not be executed (due to error).
+			if (HasDeferredComponentRegistration())
+			{
+				RegisterAllComponents();
+			}
 		}
 	}
 	else
@@ -968,7 +1017,10 @@ UActorComponent* AActor::AddComponent(FName TemplateName, bool bManualAttachment
 		}
 
 		// Register component, which will create physics/rendering state, now component is in correct position
-		NewActorComp->RegisterComponent();
+		if (NewActorComp->bAutoRegister)
+		{
+			NewActorComp->RegisterComponent();
+		}
 
 		UWorld* World = GetWorld();
 		if (!bRunningUserConstructionScript && World && bIsSceneComponent)
diff --git a/Engine/Source/Runtime/Engine/Private/BlueprintSetLibrary.cpp b/Engine/Source/Runtime/Engine/Private/BlueprintSetLibrary.cpp
new file mode 100644
--- /dev/null
+++ b/Engine/Source/Runtime/Engine/Private/BlueprintSetLibrary.cpp
@@ -0,0 +1,168 @@
+// Copyright 1998-2016 Epic Games, Inc. All Rights Reserved.
+
+#include "EnginePrivate.h"
+#include "Blueprint/BlueprintSupport.h"
+#include "Kismet/BlueprintSetLibrary.h"
+#include "Kismet/KismetArrayLibrary.h"
+#include "Templates/HasGetTypeHash.h"
+
+UBlueprintSetLibrary::UBlueprintSetLibrary(const FObjectInitializer& ObjectInitializer)
+	: Super(ObjectInitializer)
+{
+}
+
+void UBlueprintSetLibrary::GenericSet_Add(const void* TargetSet, const USetProperty* SetProperty, const void* ItemPtr)
+{
+	if (TargetSet)
+	{
+		FScriptSetHelper SetHelper(SetProperty, TargetSet);
+		SetHelper.AddElement(ItemPtr);
+	}
+}
+
+void UBlueprintSetLibrary::GenericSet_AddItems(const void* TargetSet, const USetProperty* SetProperty, const void* TargetArray, const UArrayProperty* ArrayProperty)
+{
+	if(TargetSet && TargetArray)
+	{
+		FScriptArrayHelper ArrayHelper(ArrayProperty, TargetArray);
+
+		const int32 Size = ArrayHelper.Num();
+		for (int32 I = 0; I < Size; ++I)
+		{
+			GenericSet_Add(TargetSet, SetProperty, ArrayHelper.GetRawPtr(I));
+		}
+	}
+}
+
+void UBlueprintSetLibrary::GenericSet_Remove(const void* TargetSet, const USetProperty* SetProperty, const void* ItemPtr)
+{
+	if (TargetSet)
+	{
+		FScriptSetHelper SetHelper(SetProperty, TargetSet);
+		SetHelper.RemoveElement(ItemPtr);
+	}
+}
+
+void UBlueprintSetLibrary::GenericSet_RemoveItems(const void* TargetSet, const USetProperty* SetProperty, const void* TargetArray, const UArrayProperty* ArrayProperty)
+{
+	if (TargetSet && TargetArray)
+	{
+		FScriptArrayHelper ArrayHelper(ArrayProperty, TargetArray);
+
+		const int32 Size = ArrayHelper.Num();
+		for (int32 I = 0; I < Size; ++I)
+		{
+			GenericSet_Remove(TargetSet, SetProperty, ArrayHelper.GetRawPtr(I));
+		}
+	}
+}
+
+void UBlueprintSetLibrary::GenericSet_ToArray(const void* TargetSet, const USetProperty* SetProperty, void* TargetArray, const UArrayProperty* ArrayProperty)
+{
+	if (TargetSet && TargetArray)
+	{
+		FScriptSetHelper SetHelper(SetProperty, TargetSet);
+
+		const int32 Size = SetHelper.Num();
+		for (int32 I = 0; I < Size; ++I)
+		{
+			UKismetArrayLibrary::GenericArray_Add(TargetArray, ArrayProperty, SetHelper.GetElementPtr(I));
+		}
+	}
+}
+
+void UBlueprintSetLibrary::GenericSet_Clear(const void* TargetSet, const USetProperty* SetProperty)
+{
+	if (TargetSet)
+	{
+		FScriptSetHelper SetHelper(SetProperty, TargetSet);
+		SetHelper.EmptyElements();
+	}
+}
+
+int32 UBlueprintSetLibrary::GenericSet_Length(const void* TargetSet, const USetProperty* SetProperty)
+{
+	if (TargetSet)
+	{
+		FScriptSetHelper SetHelper(SetProperty, TargetSet);
+		return SetHelper.Num();
+	}
+
+	return 0;
+}
+
+bool UBlueprintSetLibrary::GenericSet_Contains(const void* TargetSet, const USetProperty* SetProperty, const void* ItemToFind)
+{
+	if (TargetSet)
+	{
+		FScriptSetHelper SetHelper(SetProperty, TargetSet);
+		UProperty* ElementProp = SetProperty->ElementProp;
+
+		return SetHelper.FindElementFromHash(ItemToFind) != nullptr;
+	}
+
+	return false;
+}
+
+void UBlueprintSetLibrary::GenericSet_Intersect(const void* SetA, const USetProperty* SetPropertyA, const void* SetB, const USetProperty* SetPropertyB, const void* SetResult, const USetProperty* SetPropertyResult)
+{
+	if (SetA && SetB && SetResult)
+	{
+		FScriptSetHelper SetHelperA(SetPropertyA, SetA);
+		FScriptSetHelper SetHelperB(SetPropertyB, SetB);
+		FScriptSetHelper SetHelperResult(SetPropertyResult, SetResult);
+
+		const int32 Size = SetHelperA.Num();
+		for (int32 I = 0; I < Size; ++I)
+		{
+			const void* EntryInA = SetHelperA.GetElementPtr(I);
+			if (SetHelperB.FindElementFromHash(EntryInA) != nullptr)
+			{
+				SetHelperResult.AddElement(EntryInA);
+			}
+		}
+	}
+}
+
+void UBlueprintSetLibrary::GenericSet_Union(const void* SetA, const USetProperty* SetPropertyA, const void* SetB, const USetProperty* SetPropertyB, const void* SetResult, const USetProperty* SetPropertyResult)
+{
+	if (SetA && SetB && SetResult)
+	{
+		FScriptSetHelper SetHelperA(SetPropertyA, SetA);
+		FScriptSetHelper SetHelperB(SetPropertyB, SetB);
+		FScriptSetHelper SetHelperResult(SetPropertyResult, SetResult);
+
+		const int32 SizeA = SetHelperA.Num();
+		for (int32 I = 0; I < SizeA; ++I)
+		{
+			SetHelperResult.AddElement(SetHelperA.GetElementPtr(I));
+		}
+
+		const int32 SizeB = SetHelperB.Num();
+		for (int32 I = 0; I < SizeB; ++I)
+		{
+			SetHelperResult.AddElement(SetHelperB.GetElementPtr(I));
+		}
+	}
+}
+
+void UBlueprintSetLibrary::GenericSet_Difference(const void* SetA, const USetProperty* SetPropertyA, const void* SetB, const USetProperty* SetPropertyB, const void* SetResult, const USetProperty* SetPropertyResult)
+{
+	if (SetA && SetB && SetResult)
+	{
+		FScriptSetHelper SetHelperA(SetPropertyA, SetA);
+		FScriptSetHelper SetHelperB(SetPropertyB, SetB);
+		FScriptSetHelper SetHelperResult(SetPropertyResult, SetResult);
+
+		const int32 Size = SetHelperA.Num();
+		for (int32 I = 0; I < Size; ++I)
+		{
+			const void* EntryInA = SetHelperA.GetElementPtr(I);
+			if (SetHelperB.FindElementFromHash(EntryInA) == nullptr)
+			{
+				SetHelperResult.AddElement(EntryInA);
+			}
+		}
+	}
+}
+
diff --git a/Engine/Source/Runtime/Engine/Private/Components/ActorComponent.cpp b/Engine/Source/Runtime/Engine/Private/Components/ActorComponent.cpp
--- a/Engine/Source/Runtime/Engine/Private/Components/ActorComponent.cpp
+++ b/Engine/Source/Runtime/Engine/Private/Components/ActorComponent.cpp
@@ -106,8 +106,6 @@ FActorComponentCreatePhysicsSignature UActorComponent::CreatePhysicsDelegate;
 // Destroy Physics global delegate
 FActorComponentDestroyPhysicsSignature UActorComponent::DestroyPhysicsDelegate;
 
-const FString UActorComponent::ComponentTemplateNameSuffix(TEXT("_GEN_VARIABLE"));
-
 UActorComponent::UActorComponent(const FObjectInitializer& ObjectInitializer /*= FObjectInitializer::Get()*/)
 	: Super(ObjectInitializer)
 {
diff --git a/Engine/Source/Runtime/Engine/Private/Components/SceneComponent.cpp b/Engine/Source/Runtime/Engine/Private/Components/SceneComponent.cpp
--- a/Engine/Source/Runtime/Engine/Private/Components/SceneComponent.cpp
+++ b/Engine/Source/Runtime/Engine/Private/Components/SceneComponent.cpp
@@ -374,7 +374,7 @@ static void UpdateAttachedIsEditorOnly(USceneComponent* ComponentThatChanged)
 	const int32 NumComponentsChanged = SetDescendantIsEditorOnly(ComponentThatChanged);
 
 	// if we altered any components (other than the ones selected), then notify the user
-	if (NumComponentsChanged > 0 && !ComponentThatChanged->GetName().EndsWith(UActorComponent::ComponentTemplateNameSuffix))
+	if (NumComponentsChanged > 0 && !ComponentThatChanged->HasAllFlags(RF_ArchetypeObject))
 	{
 		FText NotificationText = LOCTEXT("IsEditorOnlyAlteredSingularNotification", "Caused 1 component to also change its IsEditorOnly behaviour");
 		if (NumComponentsChanged > 1)
diff --git a/Engine/Source/Runtime/Engine/Private/EdGraph/EdGraphPin.cpp b/Engine/Source/Runtime/Engine/Private/EdGraph/EdGraphPin.cpp
--- a/Engine/Source/Runtime/Engine/Private/EdGraph/EdGraphPin.cpp
+++ b/Engine/Source/Runtime/Engine/Private/EdGraph/EdGraphPin.cpp
@@ -199,6 +199,14 @@ enum class EPinResolveType : uint8
 	ReferencePassThroughConnection
 };
 
+enum class EPinResolveResult : uint8
+{
+	FailedParse, // failed to even parse the text buffer, meaning we have to bail out
+	FailedSafely, // simply failed to resolve the object referenced by the text buffer, meaning we have to disconnect
+	Deferred, // owning object was found, but referenced pin is not presenet yet
+	Suceeded // immediately resolved
+};
+
 struct FUnresolvedPinData
 {
 	UEdGraphPin* ReferencingPin;
@@ -274,6 +282,7 @@ namespace PinHelpers
 	// Don't need bDisplayAsMutableRef's name because it is transient
 	// Don't need bWasTrashed's name because it is transient
 #endif
+	static EPinResolveResult ImportText_PinReference(const TCHAR*& Buffer, UEdGraphPin*& PinRef, int32 ArrayIdx, UEdGraphPin* RequestingPin, EPinResolveType ResolveType);
 }
 
 /////////////////////////////////////////////////////
@@ -912,11 +921,11 @@ bool UEdGraphPin::ImportTextItem(const TCHAR*& Buffer, int32 PortFlags, class UO
 		}
 		else if (PropertyToken == PinHelpers::ParentPinName)
 		{
-			bParseSuccess = UEdGraphPin::ImportText_PinReference(Buffer, ParentPin, INDEX_NONE, this, EPinResolveType::ParentPin);
+			bParseSuccess = PinHelpers::ImportText_PinReference(Buffer, ParentPin, INDEX_NONE, this, EPinResolveType::ParentPin) != EPinResolveResult::FailedParse;
 		}
 		else if (PropertyToken == PinHelpers::ReferencePassThroughConnectionName)
 		{
-			bParseSuccess = UEdGraphPin::ImportText_PinReference(Buffer, ReferencePassThroughConnection, INDEX_NONE, this, EPinResolveType::ReferencePassThroughConnection);
+			bParseSuccess = PinHelpers::ImportText_PinReference(Buffer, ReferencePassThroughConnection, INDEX_NONE, this, EPinResolveType::ReferencePassThroughConnection) != EPinResolveResult::FailedParse;
 		}
 #if WITH_EDITORONLY_DATA
 		else if (PropertyToken == PinHelpers::PersistentGuidName)
@@ -989,7 +998,7 @@ bool UEdGraphPin::ImportTextItem(const TCHAR*& Buffer, int32 PortFlags, class UO
 	}
 
 	// Someone might have been waiting for this pin to be created. Let them know that this pin now exists.
-	ResolveReferencesToPin(this, false);
+	ResolveReferencesToPin(this);
 
 	return true;
 }
@@ -1618,7 +1627,7 @@ FString UEdGraphPin::ExportText_PinArray(const TArray<UEdGraphPin*>& PinArray)
 	return RetVal;
 }
 
-bool UEdGraphPin::ImportText_PinReference(const TCHAR*& Buffer, UEdGraphPin*& PinRef, int32 ArrayIdx, UEdGraphPin* RequestingPin, EPinResolveType ResolveType)
+EPinResolveResult PinHelpers::ImportText_PinReference(const TCHAR*& Buffer, UEdGraphPin*& PinRef, int32 ArrayIdx, UEdGraphPin* RequestingPin, EPinResolveType ResolveType)
 {
 	const TCHAR* BufferStart = Buffer;
 	while (*Buffer != PinHelpers::ExportTextPropDelimiter && *Buffer != ')')
@@ -1626,7 +1635,7 @@ bool UEdGraphPin::ImportText_PinReference(const TCHAR*& Buffer, UEdGraphPin*& Pi
 		if (*Buffer == 0)
 		{
 			// Parse error
-			return false;
+			return EPinResolveResult::FailedParse;
 		}
 
 		Buffer++;
@@ -1647,17 +1656,20 @@ bool UEdGraphPin::ImportText_PinReference(const TCHAR*& Buffer, UEdGraphPin*& Pi
 				if (ExistingPin)
 				{
 					PinRef = *ExistingPin;
+					return EPinResolveResult::Suceeded;
 				}
 				else
 				{
 					TArray<FUnresolvedPinData>& UnresolvedPinData = PinHelpers::UnresolvedPins.FindOrAdd(FPinResolveId(PinGuid, LocalOwningNode));
 					UnresolvedPinData.Add(FUnresolvedPinData(RequestingPin, ResolveType, ArrayIdx));
+					return EPinResolveResult::Deferred;
 				}
 			}
 		}
+		
 	}
 
-	return true;
+	return EPinResolveResult::FailedSafely;
 }
 
 bool UEdGraphPin::ImportText_PinArray(const TCHAR*& Buffer, TArray<UEdGraphPin*>& ArrayRef, UEdGraphPin* RequestingPin, EPinResolveType ResolveType)
@@ -1678,11 +1690,17 @@ bool UEdGraphPin::ImportText_PinArray(const TCHAR*& Buffer, TArray<UEdGraphPin*>
 			}
 
 			int32 NewItemIdx = ArrayRef.Add(nullptr);
-			bool bParseSuccess = ImportText_PinReference(Buffer, ArrayRef[NewItemIdx], NewItemIdx, RequestingPin, ResolveType);
-			if (!bParseSuccess)
+			EPinResolveResult ResolveResult = PinHelpers::ImportText_PinReference(Buffer, ArrayRef[NewItemIdx], NewItemIdx, RequestingPin, ResolveType);
+
+			if (ResolveResult == EPinResolveResult::FailedParse)
 			{
+				ArrayRef.Pop();
 				return false;
 			}
+			else if (ResolveResult == EPinResolveResult::FailedSafely)
+			{
+				ArrayRef.Pop();
+			}
 
 			if (*Buffer == PinHelpers::ExportTextPropDelimiter)
 			{
diff --git a/Engine/Source/Runtime/Engine/Private/InheritableComponentHandler.cpp b/Engine/Source/Runtime/Engine/Private/InheritableComponentHandler.cpp
--- a/Engine/Source/Runtime/Engine/Private/InheritableComponentHandler.cpp
+++ b/Engine/Source/Runtime/Engine/Private/InheritableComponentHandler.cpp
@@ -14,20 +14,37 @@ void UInheritableComponentHandler::PostLoad()
 {
 	Super::PostLoad();
 
-	for (int32 Index = Records.Num() - 1; Index >= 0; --Index)
+	if (!GIsDuplicatingClassForReinstancing)
 	{
-		FComponentOverrideRecord& Record = Records[Index];
-		if (Record.ComponentTemplate)
+		for (int32 Index = Records.Num() - 1; Index >= 0; --Index)
 		{
-			if (!CastChecked<UActorComponent>(Record.ComponentTemplate->GetArchetype())->IsEditableWhenInherited())
-			{
-				Record.ComponentTemplate->MarkPendingKill(); // hack needed to be able to identify if NewObject returns this back to us in the future
-				Records.RemoveAtSwap(Index);
-			}
-			else if(Record.CookedComponentInstancingData.bIsValid)
+			FComponentOverrideRecord& Record = Records[Index];
+			if (Record.ComponentTemplate)
 			{
-				// Generate "fast path" instancing data.
-				Record.CookedComponentInstancingData.LoadCachedPropertyDataForSerialization(Record.ComponentTemplate);
+				// Fix up component class on load, if it's not already set.
+				if (Record.ComponentClass == nullptr)
+				{
+					Record.ComponentClass = Record.ComponentTemplate->GetClass();
+				}
+
+				// Fix up component template name on load, if it doesn't match the original template name. Otherwise, archetype lookups will fail for this template.
+				// For example, this can occur after a component variable rename in a parent BP class, but before a child BP class with an override template is loaded.
+				UActorComponent* OriginalTemplate = Record.ComponentKey.GetOriginalTemplate();
+				if (OriginalTemplate && OriginalTemplate->GetFName() != Record.ComponentTemplate->GetFName())
+				{
+					Record.ComponentTemplate->Rename(*OriginalTemplate->GetName(), nullptr, REN_DontCreateRedirectors | REN_ForceNoResetLoaders);
+				}
+
+				if (!CastChecked<UActorComponent>(Record.ComponentTemplate->GetArchetype())->IsEditableWhenInherited())
+				{
+					Record.ComponentTemplate->MarkPendingKill(); // hack needed to be able to identify if NewObject returns this back to us in the future
+					Records.RemoveAtSwap(Index);
+				}
+				else if (Record.CookedComponentInstancingData.bIsValid)
+				{
+					// Generate "fast path" instancing data.
+					Record.CookedComponentInstancingData.LoadCachedPropertyDataForSerialization(Record.ComponentTemplate);
+				}
 			}
 		}
 	}
@@ -83,6 +100,7 @@ UActorComponent* UInheritableComponentHandler::CreateOverridenComponentTemplate(
 
 	FComponentOverrideRecord NewRecord;
 	NewRecord.ComponentKey = Key;
+	NewRecord.ComponentClass = NewComponentTemplate->GetClass();
 	NewRecord.ComponentTemplate = NewComponentTemplate;
 	Records.Add(NewRecord);
 
@@ -143,13 +161,13 @@ void UInheritableComponentHandler::ValidateTemplates()
 				}
 				else
 				{
-					UE_LOG(LogBlueprint, Log, TEXT("ValidateTemplates '%s': overridden template is unnecessary - component '%s' from '%s'"),
+					UE_LOG(LogBlueprint, Log, TEXT("ValidateTemplates '%s': overridden template is unnecessary and will be removed - component '%s' from '%s'"),
 						*GetPathNameSafe(this), *VarName.ToString(), *GetPathNameSafe(ComponentKey.GetComponentOwner()));
 				}
 			}
 			else
 			{
-				UE_LOG(LogBlueprint, Warning, TEXT("ValidateTemplates '%s': overridden template is invalid - component '%s' from '%s'"),
+				UE_LOG(LogBlueprint, Warning, TEXT("ValidateTemplates '%s': overridden template is invalid and will be removed - component '%s' from '%s'"),
 					*GetPathNameSafe(this), *VarName.ToString(), *GetPathNameSafe(ComponentKey.GetComponentOwner()));
 			}
 		}
@@ -187,7 +205,8 @@ bool UInheritableComponentHandler::IsRecordValid(const FComponentOverrideRecord&
 
 	if (!Record.ComponentTemplate)
 	{
-		return false;
+		// Note: We still consider the record to be valid, even if the template is missing, if we have valid class information. This typically will indicate that the template object was filtered at load time (to save space, e.g. dedicated server).
+		return Record.ComponentClass != nullptr;
 	}
 
 	if (Record.ComponentTemplate->GetOuter() != OwnerClass)
@@ -206,13 +225,9 @@ bool UInheritableComponentHandler::IsRecordValid(const FComponentOverrideRecord&
 		return false;
 	}
 
-	auto OriginalTemplate = Record.ComponentKey.GetOriginalTemplate();
-	if (!OriginalTemplate)
-	{
-		return false;
-	}
-
-	if (OriginalTemplate->GetClass() != Record.ComponentTemplate->GetClass())
+	// Note: If the original template is missing, we consider the record to be unnecessary, but not invalid.
+	UActorComponent* OriginalTemplate = Record.ComponentKey.GetOriginalTemplate();
+	if (OriginalTemplate != nullptr && OriginalTemplate->GetClass() != Record.ComponentTemplate->GetClass())
 	{
 		return false;
 	}
@@ -255,10 +270,36 @@ struct FComponentComparisonHelper
 
 bool UInheritableComponentHandler::IsRecordNecessary(const FComponentOverrideRecord& Record) const
 {
-	auto ChildComponentTemplate = Record.ComponentTemplate;
-	auto ParentComponentTemplate = FindBestArchetype(Record.ComponentKey);
-	check(ChildComponentTemplate && ParentComponentTemplate && (ParentComponentTemplate != ChildComponentTemplate));
-	return !FComponentComparisonHelper::AreIdentical(ChildComponentTemplate, ParentComponentTemplate);
+	// If the record's template was not loaded, check to see if the class information is valid.
+	if (Record.ComponentTemplate == nullptr)
+	{
+		if (Record.ComponentClass != nullptr)
+		{
+			UObject* ComponentCDO = Record.ComponentClass->GetDefaultObject();
+			if (ComponentCDO != nullptr)
+			{
+				// The record is considered necessary if the class information is valid but the template was not loaded due to client/server exclusion at load time (e.g. uncooked dedicated server).
+				return !UObject::CanCreateInCurrentContext(ComponentCDO);
+			}
+		}
+		
+		// Otherwise, we don't need to keep the record if the template is NULL.
+		return false;
+	}
+	else
+	{
+		// Consider the record to be unnecessary if the original template no longer exists.
+		UActorComponent* OriginalTemplate = Record.ComponentKey.GetOriginalTemplate();
+		if (OriginalTemplate == nullptr)
+		{
+			return false;
+		}
+	
+		auto ChildComponentTemplate = Record.ComponentTemplate;
+		auto ParentComponentTemplate = FindBestArchetype(Record.ComponentKey);
+		check(ChildComponentTemplate && ParentComponentTemplate && (ParentComponentTemplate != ChildComponentTemplate));
+		return !FComponentComparisonHelper::AreIdentical(ChildComponentTemplate, ParentComponentTemplate);
+	}
 }
 
 UActorComponent* UInheritableComponentHandler::FindBestArchetype(FComponentKey Key) const
diff --git a/Engine/Source/Runtime/Engine/Private/KismetMathLibrary.cpp b/Engine/Source/Runtime/Engine/Private/KismetMathLibrary.cpp
--- a/Engine/Source/Runtime/Engine/Private/KismetMathLibrary.cpp
+++ b/Engine/Source/Runtime/Engine/Private/KismetMathLibrary.cpp
@@ -1132,7 +1132,7 @@ FVector UKismetMathLibrary::GetVectorArrayAverage(const TArray<FVector>& Vectors
 }
 
 /** Find the unit direction vector from one position to another. */
-FVector UKismetMathLibrary::GetDirectionVector(FVector From, FVector To)
+FVector UKismetMathLibrary::GetDirectionUnitVector(FVector From, FVector To)
 {
 	return (To - From).GetSafeNormal();
 }
diff --git a/Engine/Source/Runtime/Engine/Private/ObjectLibrary.cpp b/Engine/Source/Runtime/Engine/Private/ObjectLibrary.cpp
--- a/Engine/Source/Runtime/Engine/Private/ObjectLibrary.cpp
+++ b/Engine/Source/Runtime/Engine/Private/ObjectLibrary.cpp
@@ -46,6 +46,8 @@ void UObjectLibrary::PostEditChangeProperty(FPropertyChangedEvent& PropertyChang
 					{
 						UBlueprintCore* Blueprint = Cast<UBlueprintCore>(Objects[i]);
 						BlueprintClass = Blueprint ? Blueprint->GeneratedClass : nullptr;
+						// replace BP with BPGC
+						Objects[i] = BlueprintClass;
 					}
 
 					if (!BlueprintClass)
@@ -72,6 +74,27 @@ void UObjectLibrary::PostEditChangeProperty(FPropertyChangedEvent& PropertyChang
 
 	Super::PostEditChangeProperty(PropertyChangedEvent);
 }
+
+void UObjectLibrary::PostLoad()
+{
+	Super::PostLoad();
+
+	if (bHasBlueprintClasses)
+	{
+		// replace BP with BPGC
+		for (int32 i = 0; i < Objects.Num(); i++)
+		{
+			UBlueprintCore* Blueprint = Cast<UBlueprintCore>(Objects[i]);
+			if (Blueprint)
+			{
+				UClass* BlueprintClass = Blueprint->GeneratedClass;
+				Objects[i] = (BlueprintClass && BlueprintClass->IsChildOf(ObjectBaseClass))
+					? BlueprintClass
+					: nullptr;
+			}
+		}
+	}
+}
 #endif // WITH_EDITOR
 
 class UObjectLibrary* UObjectLibrary::CreateLibrary(UClass* InBaseClass, bool bInHasBlueprintClasses, bool bInUseWeak)
diff --git a/Engine/Source/Runtime/Engine/Private/SCS_Node.cpp b/Engine/Source/Runtime/Engine/Private/SCS_Node.cpp
--- a/Engine/Source/Runtime/Engine/Private/SCS_Node.cpp
+++ b/Engine/Source/Runtime/Engine/Private/SCS_Node.cpp
@@ -93,19 +93,32 @@ UActorComponent* USCS_Node::ExecuteNodeOnActor(AActor* Actor, USceneComponent* P
 	if (ActualComponentTemplateData && ActualComponentTemplateData->bIsValid)
 	{
 		// Use cooked instancing data if valid (fast path).
-		NewActorComp = Actor->CreateComponentFromTemplateData(ActualComponentTemplateData, VariableName);
+		NewActorComp = Actor->CreateComponentFromTemplateData(ActualComponentTemplateData, InternalVariableName);
 	}
 	else if (UActorComponent* ActualComponentTemplate = GetActualComponentTemplate(ActualBPGC))
 	{
-		NewActorComp = Actor->CreateComponentFromTemplate(ActualComponentTemplate, VariableName);
+		NewActorComp = Actor->CreateComponentFromTemplate(ActualComponentTemplate, InternalVariableName);
 	}
 
 	if(NewActorComp != nullptr)
 	{
 		NewActorComp->CreationMethod = EComponentCreationMethod::SimpleConstructionScript;
+
 		// SCS created components are net addressable
 		NewActorComp->SetNetAddressable();
 
+		if (!NewActorComp->HasBeenCreated())
+		{
+			// Call function to notify component it has been created
+			NewActorComp->OnComponentCreated();
+		}
+
+		if (NewActorComp->GetIsReplicated())
+		{
+			// Make sure this component is added to owning actor's replicated list.
+			NewActorComp->SetIsReplicated(true);
+		}
+
 		// Special handling for scene components
 		USceneComponent* NewSceneComp = Cast<USceneComponent>(NewActorComp);
 		if (NewSceneComp != nullptr)
@@ -123,25 +136,26 @@ UActorComponent* USCS_Node::ExecuteNodeOnActor(AActor* Actor, USceneComponent* P
 				}
 				NewSceneComp->SetWorldTransform(WorldTransform);
 				Actor->SetRootComponent(NewSceneComp);
+
+				// This will be true if we deferred the RegisterAllComponents() call at spawn time. In that case, we can call it now since we have set a scene root.
+				if (Actor->HasDeferredComponentRegistration())
+				{
+					// Register the root component along with any components whose registration may have been deferred pending SCS execution in order to establish a root.
+					Actor->RegisterAllComponents();
+				}
 			}
 			// Otherwise, attach to parent component passed in
 			else
 			{
 				NewSceneComp->SetupAttachment(ParentComponent, AttachToName);
 			}
-		}
 
-		// Call function to notify component it has been created
-		NewActorComp->OnComponentCreated();
-
-		if (NewActorComp->GetIsReplicated())
-		{
-			// Make sure this component is added to owning actor's replicated list.
-			NewActorComp->SetIsReplicated(true);
+			// Register SCS scene components now (if necessary). Non-scene SCS component registration is deferred until after SCS execution, as there can be dependencies on the scene hierarchy.
+			USimpleConstructionScript::RegisterInstancedComponent(NewSceneComp);
 		}
 
 		// If we want to save this to a property, do it here
-		FName VarName = GetVariableName();
+		FName VarName = InternalVariableName;
 		if (VarName != NAME_None)
 		{
 			UClass* ActorClass = Actor->GetClass();
@@ -151,7 +165,7 @@ UActorComponent* USCS_Node::ExecuteNodeOnActor(AActor* Actor, USceneComponent* P
 			}
 			else
 			{
-				UE_LOG(LogBlueprint, Log, TEXT("ExecuteNodeOnActor: Couldn't find property '%s' on '%s"), *VarName.ToString(), *Actor->GetName());
+				UE_LOG(LogBlueprint, Log, TEXT("ExecuteNodeOnActor: Couldn't find property '%s' on '%s'"), *VarName.ToString(), *Actor->GetName());
 #if WITH_EDITOR
 				// If we're constructing editable components in the SCS editor, set the component instance corresponding to this node for editing purposes
 				USimpleConstructionScript* SCS = GetSCS();
@@ -319,36 +333,63 @@ bool USCS_Node::IsRootNode() const
 	return(SCS->GetRootNodes().Contains(const_cast<USCS_Node*>(this)));
 }
 
-FName USCS_Node::GetVariableName() const
+void USCS_Node::RenameComponentTemplate(UActorComponent* ComponentTemplate, const FName& NewName)
 {
-	// Name specified
-	if(VariableName != NAME_None)
-	{
-		return VariableName;
-	}
-	// Not specified, make variable based on template name
-	// Note that since SCS_Nodes should all have auto generated names, this code shouldn't be hit unless
-	// the auto naming code fails.
-	else if(ComponentTemplate != NULL)
-	{
-		FString VarString = FString::Printf(TEXT("%s_Var"), *ComponentTemplate->GetName());
-		return FName(*VarString);
-	}
-	else
+	if (ComponentTemplate != nullptr && ComponentTemplate->HasAllFlags(RF_ArchetypeObject))
 	{
-		return NAME_None;
+		// Gather all instances of the template (archetype)
+		TArray<UObject*> ArchetypeInstances;
+		ComponentTemplate->GetArchetypeInstances(ArchetypeInstances);
+
+		// Rename the component template (archetype) - note that this can be called during compile-on-load, so we include the flag not to reset the BPGC's package loader.
+		const FString NewComponentName = NewName.ToString();
+		ComponentTemplate->Rename(*(NewComponentName + USimpleConstructionScript::ComponentTemplateNameSuffix), nullptr, REN_DontCreateRedirectors | REN_ForceNoResetLoaders);
+
+		// Rename all component instances to match the updated variable name
+		for (UObject* ArchetypeInstance : ArchetypeInstances)
+		{
+			// Recursively handle inherited component template overrides. In the SCS case, this is because we must also handle these before the SCS key's variable name is changed.
+			if (ArchetypeInstance->HasAllFlags(RF_ArchetypeObject | RF_InheritableComponentTemplate))
+			{
+				RenameComponentTemplate(CastChecked<UActorComponent>(ArchetypeInstance), NewName);
+			}
+			else
+			{
+				// If this is an instanced component (i.e. owned by an Actor), ensure that we have no conflict with another instanced component belonging to the same Actor instance.
+				if (AActor* Actor = Cast<AActor>(ArchetypeInstance->GetOuter()))
+				{
+					Actor->CheckComponentInstanceName(NewName);
+				}
+
+				ArchetypeInstance->Rename(*NewComponentName, nullptr, REN_DontCreateRedirectors | REN_ForceNoResetLoaders);
+			}
+		}
 	}
 }
 
-void USCS_Node::NameWasModified()
+void USCS_Node::SetVariableName(const FName& NewName, bool bRenameTemplate)
 {
-	if(ComponentTemplate != nullptr)
+	// We need to ensure that component object names stay in sync with the variable name; this is done for 2 reasons:
+	//	1) This ensures existing instances can successfully route back to the archetype (template) object through the variable name.
+	//	2) This prevents new SCS nodes for the same component type from recycling an existing template with the original (base) name.
+	if (bRenameTemplate && ComponentTemplate != nullptr)
 	{
-		// Ensure that the template name stays in sync with the variable name; otherwise, new SCS nodes for the same component type will recycle the subobject rather than create a new instance.
-		ComponentTemplate->Rename(*(VariableName.ToString() + TEXT("_GEN_VARIABLE")), nullptr, REN_DontCreateRedirectors|REN_ForceNoResetLoaders);
+		// This must be called BEFORE we change the internal variable name; otherwise it will fail to find any instances of the archetype!
+		RenameComponentTemplate(ComponentTemplate, NewName);
 	}
 
-	OnNameChangedExternal.ExecuteIfBound(VariableName);
+	InternalVariableName = NewName;
+
+	// >>> Backwards Compatibility: Support existing projects/tools that might be reading the variable name directly. This can be removed in a future release.
+PRAGMA_DISABLE_DEPRECATION_WARNINGS
+	VariableName = InternalVariableName;
+PRAGMA_ENABLE_DEPRECATION_WARNINGS
+	// <<< End Backwards Compatibility
+}
+
+void USCS_Node::NameWasModified()
+{
+	OnNameChangedExternal.ExecuteIfBound(InternalVariableName);
 }
 
 void USCS_Node::SetOnNameChanged( const FSCSNodeNameChanged& OnChange )
@@ -397,6 +438,30 @@ void USCS_Node::RemoveMetaData(const FName& Key)
 	}
 }
 
+void USCS_Node::Serialize(FArchive& Ar)
+{
+	Super::Serialize(Ar);
+
+	if (Ar.IsLoading())
+	{
+		if (Ar.IsPersistent() && !Ar.HasAnyPortFlags(PPF_Duplicate | PPF_DuplicateForPIE))
+		{
+			// >>> Backwards Compatibility: Support existing projects/tools that might be reading the variable name directly. This can be removed in a future release.
+PRAGMA_DISABLE_DEPRECATION_WARNINGS
+			VariableName = InternalVariableName;
+PRAGMA_ENABLE_DEPRECATION_WARNINGS
+			// <<< End Backwards Compatibility
+
+			// Fix up the component class property, if it has not already been set.
+			// Note: This is done here, instead of in PostLoad(), because it needs to be set before Blueprint class compilation.
+			if (ComponentClass == nullptr && ComponentTemplate != nullptr)
+			{
+				ComponentClass = ComponentTemplate->GetClass();
+			}
+		}
+	}
+}
+
 void USCS_Node::PostLoad()
 {
 	Super::PostLoad();
@@ -412,27 +477,33 @@ void USCS_Node::PostLoad()
 	}
 }
 
+
 #if WITH_EDITOR
 void USCS_Node::SetParent(USCS_Node* InParentNode)
 {
-	check(InParentNode != NULL);
-	check(InParentNode->GetSCS() != NULL);
-	check(InParentNode->GetSCS()->GetBlueprint() != NULL);
-	check(InParentNode->GetSCS()->GetBlueprint()->GeneratedClass != NULL);
-
-	const FName NewParentComponentOrVariableName = InParentNode->VariableName;
-	const FName NewParentComponentOwnerClassName = InParentNode->GetSCS()->GetBlueprint()->GeneratedClass->GetFName();
-
-	// Only modify if it differs from current
-	if(bIsParentComponentNative
-		|| ParentComponentOrVariableName != NewParentComponentOrVariableName
-		|| ParentComponentOwnerClassName != NewParentComponentOwnerClassName)
+	ensure(InParentNode);
+	USimpleConstructionScript* ParentSCS = InParentNode ? InParentNode->GetSCS() : nullptr;
+	ensure(ParentSCS);
+	UBlueprint* ParentBlueprint = ParentSCS ? ParentSCS->GetBlueprint() : nullptr;
+	ensure(ParentBlueprint);
+	UClass* ParentBlueprintGeneratedClass = ParentBlueprint ? ParentBlueprint->GeneratedClass : nullptr;
+
+	if (ParentBlueprintGeneratedClass && InParentNode)
 	{
-		Modify();
+		const FName NewParentComponentOrVariableName = InParentNode->GetVariableName();
+		const FName NewParentComponentOwnerClassName = ParentBlueprintGeneratedClass->GetFName();
 
-		bIsParentComponentNative = false;
-		ParentComponentOrVariableName = NewParentComponentOrVariableName;
-		ParentComponentOwnerClassName = NewParentComponentOwnerClassName;
+		// Only modify if it differs from current
+		if (bIsParentComponentNative
+			|| ParentComponentOrVariableName != NewParentComponentOrVariableName
+			|| ParentComponentOwnerClassName != NewParentComponentOwnerClassName)
+		{
+			Modify();
+
+			bIsParentComponentNative = false;
+			ParentComponentOrVariableName = NewParentComponentOrVariableName;
+			ParentComponentOwnerClassName = NewParentComponentOwnerClassName;
+		}
 	}
 }
 
@@ -506,7 +577,7 @@ USceneComponent* USCS_Node::GetParentComponentTemplate(UBlueprint* InBlueprint)
 					for(int32 ParentNodeIndex = 0; ParentNodeIndex < ParentSCSNodes.Num(); ++ParentNodeIndex)
 					{
 						USceneComponent* CompTemplate = Cast<USceneComponent>(ParentSCSNodes[ParentNodeIndex]->ComponentTemplate);
-						if(CompTemplate != NULL && ParentSCSNodes[ParentNodeIndex]->VariableName == ParentComponentOrVariableName)
+						if(CompTemplate != NULL && ParentSCSNodes[ParentNodeIndex]->GetVariableName() == ParentComponentOrVariableName)
 						{
 							// Found a match; this is our parent, we're done
 							ParentComponentTemplate = Cast<USceneComponent>(ParentSCSNodes[ParentNodeIndex]->GetActualComponentTemplate(Cast<UBlueprintGeneratedClass>(InBlueprint->GeneratedClass)));
@@ -526,9 +597,9 @@ void USCS_Node::ValidateGuid()
 	// Backward compatibility:
 	// The guid for the node should be always the same (event when it's not saved). 
 	// The guid is created in an deterministic way using persistent name.
-	if (!VariableGuid.IsValid() && (VariableName != NAME_None))
+	if (!VariableGuid.IsValid() && (InternalVariableName != NAME_None))
 	{
-		const FString HashString = VariableName.ToString();
+		const FString HashString = InternalVariableName.ToString();
 		ensure(HashString.Len());
 
 		const uint32 BufferLength = HashString.Len() * sizeof(HashString[0]);
diff --git a/Engine/Source/Runtime/Engine/Private/SimpleConstructionScript.cpp b/Engine/Source/Runtime/Engine/Private/SimpleConstructionScript.cpp
--- a/Engine/Source/Runtime/Engine/Private/SimpleConstructionScript.cpp
+++ b/Engine/Source/Runtime/Engine/Private/SimpleConstructionScript.cpp
@@ -2,6 +2,7 @@
 
 #include "EnginePrivate.h"
 #include "Engine/SCS_Node.h"
+#include "BlueprintsObjectVersion.h"
 #include "BlueprintUtilities.h"
 #if WITH_EDITOR
 #include "Kismet2/CompilerResultsLog.h"
@@ -14,29 +15,8 @@
 //////////////////////////////////////////////////////////////////////////
 // USimpleConstructionScript
 
-namespace
-{
-	// Helper method to register instanced components post-construction
-	void RegisterInstancedComponent(UActorComponent* InstancedComponent)
-	{
-		if (!InstancedComponent->IsRegistered())
-		{
-			InstancedComponent->RegisterComponent();
-
-			// If this is a scene component, recursively register any child components as well
-			USceneComponent* InstancedSceneComponent = Cast<USceneComponent>(InstancedComponent);
-			if(InstancedSceneComponent != nullptr)
-			{
-				// Make a copy of the array here in case something alters the AttachChildren array during registration (e.g. Physics)
-				TArray<USceneComponent*> AttachChildren = InstancedSceneComponent->GetAttachChildren();
-				for (USceneComponent* InstancedChildComponent : AttachChildren)
-				{
-					RegisterInstancedComponent(InstancedChildComponent);
-				}
-			}
-		}
-	}
-}
+// We append this suffix to template object names because the UObjectProperty we create at compile time will also be outered to the generated Blueprint class, and because we need cooking to be deterministic with respect to template object names.
+const FString USimpleConstructionScript::ComponentTemplateNameSuffix(TEXT("_GEN_VARIABLE"));
 
 USimpleConstructionScript::USimpleConstructionScript(const FObjectInitializer& ObjectInitializer)
 	: Super(ObjectInitializer)
@@ -157,70 +137,74 @@ void USimpleConstructionScript::PostLoad()
 		return;
 	}
 
-	for (USCS_Node* Node : GetAllNodes())
+	// This pass is not needed during reinstancing.
+	if (!GIsDuplicatingClassForReinstancing)
 	{
-		// Fix up any uninitialized category names
-		if(Node->CategoryName.IsEmpty())
+		for (USCS_Node* Node : GetAllNodes())
 		{
-			Node->CategoryName = NSLOCTEXT("SCS", "Default", "Default");
-		}
-
-		// Fix up components that may have switched from scene to non-scene type and vice-versa
-		if(Node->ComponentTemplate != nullptr)
-		{
-			// Fix up any component template objects whose name doesn't match the current variable name; this ensures that there is always one unique template per node.
-			FString VariableName = Node->GetVariableName().ToString();
-			FString ComponentTemplateName = Node->ComponentTemplate->GetName();
-			if(ComponentTemplateName.EndsWith(UActorComponent::ComponentTemplateNameSuffix) && !ComponentTemplateName.StartsWith(VariableName) && !GIsDuplicatingClassForReinstancing)
+			// Fix up any uninitialized category names
+			if (Node->CategoryName.IsEmpty())
 			{
-				Node->ComponentTemplate->ConditionalPostLoad();
-				Node->ComponentTemplate = static_cast<UActorComponent*>(StaticDuplicateObject(Node->ComponentTemplate, Node->ComponentTemplate->GetOuter(), *(VariableName + UActorComponent::ComponentTemplateNameSuffix)));
+				Node->CategoryName = NSLOCTEXT("SCS", "Default", "Default");
 			}
 
-			// Check to see if switched from scene to a non-scene component type
-			if (!Node->ComponentTemplate->IsA<USceneComponent>())
+			// Fix up components that may have switched from scene to non-scene type and vice-versa
+			if (Node->ComponentTemplate != nullptr)
 			{
-				// Otherwise, check to see if switched from scene to non-scene component type
-				int32 RootNodeIndex = INDEX_NONE;
-				if(!RootNodes.Find(Node, RootNodeIndex))
+				// Fix up any component template objects whose name doesn't match the current variable name; this ensures that there is always one unique template per node.
+				FString VariableName = Node->GetVariableName().ToString();
+				FString ComponentTemplateName = Node->ComponentTemplate->GetName();
+				if (ComponentTemplateName.EndsWith(ComponentTemplateNameSuffix) && !ComponentTemplateName.StartsWith(VariableName))
 				{
-					// Move the node into the root set if it's currently in the scene hierarchy
-					USCS_Node* ParentNode = FindParentNode(Node);
-					if(ParentNode != nullptr)
-					{
-						ParentNode->RemoveChildNode(Node);
-					}
-
-					RootNodes.Add(Node);
+					Node->ComponentTemplate->ConditionalPostLoad();
+					Node->ComponentTemplate = static_cast<UActorComponent*>(StaticDuplicateObject(Node->ComponentTemplate, Node->ComponentTemplate->GetOuter(), *(VariableName + ComponentTemplateNameSuffix)));
 				}
-				else
+
+				// Check to see if switched from scene to a non-scene component type
+				if (!Node->ComponentTemplate->IsA<USceneComponent>())
 				{
-					// Otherwise, if it's a root node, promote one of its children (if any) to take its place
-					int32 PromoteIndex = FindPromotableChildNodeIndex(Node);
-					if(PromoteIndex != INDEX_NONE)
+					// Otherwise, check to see if switched from scene to non-scene component type
+					int32 RootNodeIndex = INDEX_NONE;
+					if (!RootNodes.Find(Node, RootNodeIndex))
+					{
+						// Move the node into the root set if it's currently in the scene hierarchy
+						USCS_Node* ParentNode = FindParentNode(Node);
+						if (ParentNode != nullptr)
+						{
+							ParentNode->RemoveChildNode(Node);
+						}
+
+						RootNodes.Add(Node);
+					}
+					else
 					{
-						// Remove it as a child node
-						USCS_Node* ChildToPromote = Node->GetChildNodes()[PromoteIndex];
-						Node->RemoveChildNodeAt(PromoteIndex, false);
+						// Otherwise, if it's a root node, promote one of its children (if any) to take its place
+						int32 PromoteIndex = FindPromotableChildNodeIndex(Node);
+						if (PromoteIndex != INDEX_NONE)
+						{
+							// Remove it as a child node
+							USCS_Node* ChildToPromote = Node->GetChildNodes()[PromoteIndex];
+							Node->RemoveChildNodeAt(PromoteIndex, false);
 
-						// Insert it as a root node just before its prior parent node; this way if it switches back to a scene type it won't supplant the new root we've just created
-						RootNodes.Insert(ChildToPromote, RootNodeIndex);
+							// Insert it as a root node just before its prior parent node; this way if it switches back to a scene type it won't supplant the new root we've just created
+							RootNodes.Insert(ChildToPromote, RootNodeIndex);
 
-						// Append previous root node's children to the new root
-						ChildToPromote->MoveChildNodes(Node);
+							// Append previous root node's children to the new root
+							ChildToPromote->MoveChildNodes(Node);
 
-						// Copy any previous external attachment info from the previous root node
-						ChildToPromote->bIsParentComponentNative = Node->bIsParentComponentNative;
-						ChildToPromote->ParentComponentOrVariableName = Node->ParentComponentOrVariableName;
-						ChildToPromote->ParentComponentOwnerClassName = Node->ParentComponentOwnerClassName;
-					}
+							// Copy any previous external attachment info from the previous root node
+							ChildToPromote->bIsParentComponentNative = Node->bIsParentComponentNative;
+							ChildToPromote->ParentComponentOrVariableName = Node->ParentComponentOrVariableName;
+							ChildToPromote->ParentComponentOwnerClassName = Node->ParentComponentOwnerClassName;
+						}
 
-					// Clear info for any previous external attachment if set
-					if(Node->ParentComponentOrVariableName != NAME_None)
-					{
-						Node->bIsParentComponentNative = false;
-						Node->ParentComponentOrVariableName = NAME_None;
-						Node->ParentComponentOwnerClassName = NAME_None;
+						// Clear info for any previous external attachment if set
+						if (Node->ParentComponentOrVariableName != NAME_None)
+						{
+							Node->bIsParentComponentNative = false;
+							Node->ParentComponentOrVariableName = NAME_None;
+							Node->ParentComponentOwnerClassName = NAME_None;
+						}
 					}
 				}
 			}
@@ -228,11 +212,15 @@ void USimpleConstructionScript::PostLoad()
 	}
 #endif // WITH_EDITOR
 
-	// Fix up native/inherited parent attachments, in case anything has changed
-	FixupRootNodeParentReferences();
+	// Skip validation when reinstancing.
+	if (!GIsDuplicatingClassForReinstancing)
+	{
+		// Fix up native/inherited parent attachments, in case anything has changed
+		FixupRootNodeParentReferences();
 
-	// Ensure that we have a valid scene root
-	ValidateSceneRootNodes();
+		// Ensure that we have a valid scene root
+		ValidateSceneRootNodes();
+	}
 
 	// Reset non-native "root" scene component scale values, prior to the change in which
 	// we began applying custom scale values to root components at construction time. This
@@ -418,7 +406,7 @@ void USimpleConstructionScript::FixupRootNodeParentReferences()
 						// Attempt to locate a match by searching all the nodes that belong to the parent Blueprint's SCS
 						for (USCS_Node* ParentNode : ParentClass->SimpleConstructionScript->GetAllNodes())
 						{
-							if (ParentNode != nullptr && ParentNode->VariableName == RootNode->ParentComponentOrVariableName)
+							if (ParentNode != nullptr && ParentNode->GetVariableName() == RootNode->ParentComponentOrVariableName)
 							{
 								bWasFound = true;
 								break;
@@ -453,45 +441,38 @@ void USimpleConstructionScript::FixupRootNodeParentReferences()
 	FixupSceneNodeHierarchy();
 }
 
-void USimpleConstructionScript::ExecuteScriptOnActor(AActor* Actor, const FTransform& RootTransform, bool bIsDefaultTransform)
+void USimpleConstructionScript::RegisterInstancedComponent(UActorComponent* InstancedComponent)
+{
+	// If this is a scene component, recursively register parent attachments within the actor's scene hierarchy first.
+	if (USceneComponent* SceneComponent = Cast<USceneComponent>(InstancedComponent))
+	{
+		USceneComponent* ParentComponent = SceneComponent->GetAttachParent();
+		if (ParentComponent != nullptr
+			&& ParentComponent->GetOwner() == SceneComponent->GetOwner()
+			&& !ParentComponent->IsRegistered())
+		{
+			RegisterInstancedComponent(ParentComponent);
+		}
+	}
+
+	if (InstancedComponent != nullptr && !InstancedComponent->IsRegistered() && InstancedComponent->bAutoRegister && !InstancedComponent->IsPendingKill())
+	{
+		InstancedComponent->RegisterComponent();
+	}
+}
+
+void USimpleConstructionScript::ExecuteScriptOnActor(AActor* Actor, const TInlineComponentArray<USceneComponent*>& NativeSceneComponents, const FTransform& RootTransform, bool bIsDefaultTransform)
 {
 	if(RootNodes.Num() > 0)
 	{
-		TSet<UActorComponent*> AllComponentsCreatedBySCS;
-		TInlineComponentArray<UActorComponent*> InstancedComponents;
+		// Get the given actor's root component (can be NULL).
+		USceneComponent* RootComponent = Actor->GetRootComponent();
+
 		for(auto NodeIt = RootNodes.CreateIterator(); NodeIt; ++NodeIt)
 		{
 			USCS_Node* RootNode = *NodeIt;
 			if(RootNode != nullptr)
 			{
-				// Get all native scene components
-				TInlineComponentArray<USceneComponent*> Components;
-				Actor->GetComponents(Components);
-				for (int32 Index = Components.Num()-1; Index >= 0; --Index)
-				{
-					USceneComponent* SceneComponent = Components[Index];
-					if (SceneComponent->CreationMethod == EComponentCreationMethod::Instance)
-					{
-						Components.RemoveAt(Index);
-					}
-					else
-					{
-						// Handle the native sub-component of an instance component case
-						USceneComponent* ParentSceneComponent = SceneComponent->GetTypedOuter<USceneComponent>();
-						if (ParentSceneComponent && ParentSceneComponent->CreationMethod == EComponentCreationMethod::Instance)
-						{
-							Components.RemoveAt(Index);
-						}
-					}
-				}
-
-				// Get the native root component; if it's not set, the first native scene component will be used as root. This matches what's done in the SCS editor.
-				USceneComponent* RootComponent = Actor->GetRootComponent();
-				if(RootComponent == nullptr && Components.Num() > 0)
-				{
-					RootComponent = Components[0];
-				}
-
 				// If the root node specifies that it has a parent
 				USceneComponent* ParentComponent = nullptr;
 				if(RootNode->ParentComponentOrVariableName != NAME_None)
@@ -500,15 +481,17 @@ void USimpleConstructionScript::ExecuteScriptOnActor(AActor* Actor, const FTrans
 					UClass* ActorClass = Actor->GetClass();
 					check(ActorClass != nullptr);
 
-					// If the root node is parented to a "native" component (i.e. in the 'Components' array)
+					// If the root node is parented to a "native" component (i.e. in the 'NativeSceneComponents' array)
 					if(RootNode->bIsParentComponentNative)
 					{
-						for(int32 CompIndex = 0; CompIndex < Components.Num(); ++CompIndex)
+						for(USceneComponent* NativeSceneComponent : NativeSceneComponents)
 						{
-							// If we found a match, remember the index
-							if(Components[CompIndex]->GetFName() == RootNode->ParentComponentOrVariableName)
+							check(NativeSceneComponent != nullptr);
+
+							// If we found a match, remember it
+							if(NativeSceneComponent->GetFName() == RootNode->ParentComponentOrVariableName)
 							{
-								ParentComponent = Components[CompIndex];
+								ParentComponent = NativeSceneComponent;
 								break;
 							}
 						}
@@ -526,42 +509,11 @@ void USimpleConstructionScript::ExecuteScriptOnActor(AActor* Actor, const FTrans
 				}
 
 				// Create the new component instance and any child components it may have
-				UActorComponent* InstancedComponent = RootNode->ExecuteNodeOnActor(Actor, ParentComponent != nullptr ? ParentComponent : RootComponent, &RootTransform, bIsDefaultTransform);
-				if(InstancedComponent != nullptr)
-				{
-					InstancedComponents.Add(InstancedComponent);
-				}
-
-				// get list of every component SCS created, in case some of them aren't in the attachment hierarchy any more (e.g. rigid bodies)
-				TInlineComponentArray<USceneComponent*> ComponentsAfterSCS;
-				Actor->GetComponents(ComponentsAfterSCS);
-				for (USceneComponent* C : ComponentsAfterSCS)
-				{
-					if (Components.Contains(C) == false)
-					{
-						AllComponentsCreatedBySCS.Add(C);
-					}
-				}
-			}
-		}
-
-		// Register all instanced SCS components once SCS execution has finished; sorted in order to register the scene component hierarchy first, followed by the remaining actor components (in case they happen to depend on something in the scene hierarchy)
-		InstancedComponents.Sort([](const UActorComponent& A, const UActorComponent& B) { return A.IsA<USceneComponent>(); });
-		for(auto InstancedComponent : InstancedComponents)
-		{
-			RegisterInstancedComponent(InstancedComponent);
-		}
-
-		// now that the instanced components in the attachment hierarchy are registered, register any other components that SCS made but aren't in the attachment hierarchy for whatever reason.
-		for (auto C : AllComponentsCreatedBySCS)
-		{
-			if (C->IsRegistered() == false)
-			{
-				C->RegisterComponent();
+				RootNode->ExecuteNodeOnActor(Actor, ParentComponent != nullptr ? ParentComponent : RootComponent, &RootTransform, bIsDefaultTransform);
 			}
 		}
 	}
-	else if(Actor->GetRootComponent() == NULL) // Must have a root component at the end of SCS, so if we don't have one already (from base class), create a SceneComponent now
+	else if(Actor->GetRootComponent() == nullptr) // Must have a root component at the end of SCS, so if we don't have one already (from base class), create a SceneComponent now
 	{
 		USceneComponent* SceneComp = NewObject<USceneComponent>(Actor);
 		SceneComp->SetFlags(RF_Transactional);
@@ -990,7 +942,7 @@ void USimpleConstructionScript::ValidateSceneRootNodes()
 							{
 								// The root node is inherited from a parent BP class.
 								ParentComponentOwnerClassName = BPClass->GetFName();
-								ParentComponentOrVariableName = SCSNode->VariableName;
+								ParentComponentOrVariableName = SCSNode->GetVariableName();
 							}
 							else
 							{
@@ -1079,9 +1031,10 @@ void USimpleConstructionScript::GenerateListOfExistingNames(TArray<FName>& Curre
 		const USCS_Node* ChildNode = ChildrenNodes[NodeIndex];
 		if (ChildNode)
 		{
-			if (ChildNode->VariableName != NAME_None)
+			const FName VariableName = ChildNode->GetVariableName();
+			if (VariableName != NAME_None)
 			{
-				CurrentNames.Add(ChildNode->VariableName);
+				CurrentNames.Add(VariableName);
 			}
 		}
 	}
@@ -1140,8 +1093,9 @@ USCS_Node* USimpleConstructionScript::CreateNodeImpl(UActorComponent* NewCompone
 {
 	auto NewNode = NewObject<USCS_Node>(this, MakeUniqueObjectName(this, USCS_Node::StaticClass()));
 	NewNode->SetFlags(RF_Transactional);
+	NewNode->ComponentClass = NewComponentTemplate->GetClass();
 	NewNode->ComponentTemplate = NewComponentTemplate;
-	NewNode->VariableName = ComponentVariableName;
+	NewNode->SetVariableName(ComponentVariableName, false);
 
 	// Note: This should match up with UEdGraphSchema_K2::VR_DefaultCategory
 	NewNode->CategoryName = NSLOCTEXT("SCS", "Default", "Default");
@@ -1159,6 +1113,9 @@ USCS_Node* USimpleConstructionScript::CreateNode(UClass* NewComponentClass, FNam
 	// note that naming logic is duplicated in CreateNodeAndRenameComponent:
 	NewComponentVariableName = GenerateNewComponentName(NewComponentClass, NewComponentVariableName);
 
+	// At this point we should have a unique, explicit name to use for the template object.
+	check(NewComponentVariableName != NAME_None);
+
 	// A bit of a hack, but by doing this we ensure that the original object isn't outered to the BPGC. That way if we undo this action later, it'll rename the template away from the BPGC.
 	// This is necessary because of our template object naming scheme that's in place to ensure deterministic cooking. We have to keep the SCS node and template object names in sync as a result,
 	// and leaving the template outered to the BPGC can lead to template object name collisions when attempting to rename the remaining SCS nodes. See USCS_Node::NameWasModified() for more details.
@@ -1168,7 +1125,7 @@ USCS_Node* USimpleConstructionScript::CreateNode(UClass* NewComponentClass, FNam
 	NewComponentTemplate->Modify();
 
 	// Now set the actual name and outer to the BPGC.
-	NewComponentTemplate->Rename(*(NewComponentVariableName.ToString() + UActorComponent::ComponentTemplateNameSuffix), Blueprint->GeneratedClass, REN_DoNotDirty|REN_DontCreateRedirectors|REN_ForceNoResetLoaders);
+	NewComponentTemplate->Rename(*(NewComponentVariableName.ToString() + ComponentTemplateNameSuffix), Blueprint->GeneratedClass, REN_DoNotDirty|REN_DontCreateRedirectors|REN_ForceNoResetLoaders);
 
 	return CreateNodeImpl(NewComponentTemplate, NewComponentVariableName);
 }
@@ -1180,8 +1137,11 @@ USCS_Node* USimpleConstructionScript::CreateNodeAndRenameComponent(UActorCompone
 	// note that naming logic is duplicated in CreateNode:
 	FName NewComponentVariableName = GenerateNewComponentName(NewComponentTemplate->GetClass());
 
+	// At this point we should have a unique, explicit name to use for the template object.
+	check(NewComponentVariableName != NAME_None);
+
 	// Relocate the instance from the transient package to the BPGC and assign it a unique object name
-	NewComponentTemplate->Rename(*(NewComponentVariableName.ToString() + UActorComponent::ComponentTemplateNameSuffix), GetBlueprint()->GeneratedClass, REN_DontCreateRedirectors | REN_DoNotDirty);
+	NewComponentTemplate->Rename(*(NewComponentVariableName.ToString() + ComponentTemplateNameSuffix), GetBlueprint()->GeneratedClass, REN_DontCreateRedirectors | REN_DoNotDirty);
 
 	return CreateNodeImpl(NewComponentTemplate, NewComponentVariableName);
 }
@@ -1209,32 +1169,35 @@ void USimpleConstructionScript::ValidateNodeVariableNames(FCompilerResultsLog& M
 	{
 		if( Node && Node->ComponentTemplate && Node != DefaultSceneRootNode )
 		{
+			FName VariableName = Node->GetVariableName();
+
 			// Replace missing or invalid component variable names
-			if( Node->VariableName == NAME_None
+			if( VariableName == NAME_None
 				|| Node->bVariableNameAutoGenerated_DEPRECATED
-				|| !FComponentEditorUtils::IsValidVariableNameString(Node->ComponentTemplate, Node->VariableName.ToString()) )
+				|| !FComponentEditorUtils::IsValidVariableNameString(Node->ComponentTemplate, VariableName.ToString()) )
 			{
-				FName OldName = Node->VariableName;
+				FName OldName = VariableName;
 
 				// Generate a new default variable name for the component.
-				Node->VariableName = GenerateNewComponentName(Node->ComponentTemplate->GetClass());
+				VariableName = GenerateNewComponentName(Node->ComponentTemplate->GetClass());
+				Node->SetVariableName(VariableName);
 				Node->bVariableNameAutoGenerated_DEPRECATED = false;
 
 				if( OldName != NAME_None )
 				{
-					FBlueprintEditorUtils::ReplaceVariableReferences(Blueprint, OldName, Node->VariableName);
+					FBlueprintEditorUtils::ReplaceVariableReferences(Blueprint, OldName, VariableName);
 
-					MessageLog.Warning(*FString::Printf(TEXT("Found a component variable with an invalid name (%s) - changed to %s."), *OldName.ToString(), *Node->VariableName.ToString()));
+					MessageLog.Warning(*FString::Printf(TEXT("Found a component variable with an invalid name (%s) - changed to %s."), *OldName.ToString(), *VariableName.ToString()));
 				}
 			}
-			else if( ParentBPNameValidator.IsValid() && ParentBPNameValidator->IsValid(Node->VariableName) != EValidatorResult::Ok )
+			else if( ParentBPNameValidator.IsValid() && ParentBPNameValidator->IsValid(VariableName) != EValidatorResult::Ok )
 			{
-				FName OldName = Node->VariableName;
+				FName OldName = VariableName;
 
-				FName NewVariableName = FBlueprintEditorUtils::FindUniqueKismetName(Blueprint, OldName.ToString());
-				FBlueprintEditorUtils::RenameMemberVariable(Blueprint, OldName, NewVariableName );
+				VariableName = FBlueprintEditorUtils::FindUniqueKismetName(Blueprint, OldName.ToString());
+				FBlueprintEditorUtils::RenameMemberVariable(Blueprint, OldName, VariableName );
 
-				MessageLog.Warning(*FString::Printf(TEXT("Found a component variable with a conflicting name (%s) - changed to %s."), *OldName.ToString(), *Node->VariableName.ToString()));
+				MessageLog.Warning(*FString::Printf(TEXT("Found a component variable with a conflicting name (%s) - changed to %s."), *OldName.ToString(), *VariableName.ToString()));
 			}
 		}
 	}
@@ -1254,10 +1217,34 @@ void USimpleConstructionScript::ValidateNodeTemplates(FCompilerResultsLog& Messa
 			}
 		}
 
-		// Couldn't find the template the Blueprint or C++ class must have been deleted out from under us
+		// Couldn't find the template - the Blueprint or C++ class may have been deleted out from under us, or it was not loaded due to client/server exclusion
 		if (Node->ComponentTemplate == nullptr)
 		{
-			RemoveNodeAndPromoteChildren(Node);
+			bool bRemoveNode = true;
+			if (Node->ComponentClass != nullptr)
+			{
+				// Don't remove the node if the template was not loaded due to client/server exclusion (i.e. if we can't instance the class within the current runtime context)
+				UObject* ComponentCDO = Node->ComponentClass->GetDefaultObject();
+				bRemoveNode = UObject::CanCreateInCurrentContext(ComponentCDO);
+			}
+			else
+			{
+				UBlueprint* Blueprint = GetBlueprint();
+				const FString BlueprintName = Blueprint != nullptr ? Blueprint->GetName() : FString();
+				MessageLog.Warning(*FString::Printf(TEXT("Component class is not set for '%s' - this component will not be instanced, and additional warnings or errors may occur when compiling Blueprint '%s'."), *Node->GetVariableName().ToString(), *BlueprintName));
+
+				if (GetLinkerCustomVersion(FBlueprintsObjectVersion::GUID) < FBlueprintsObjectVersion::SCSHasComponentTemplateClass
+					&& (IsRunningDedicatedServer() || IsRunningClientOnly()))
+				{
+					const FString BlueprintPathName = Blueprint != nullptr ? Blueprint->GetPathName() : FString();
+					MessageLog.Note(*FString::Printf(TEXT("Try launching the editor and resaving '%s' in order to fix this."), *BlueprintPathName));
+				}
+			}
+
+			if (bRemoveNode)
+			{
+				RemoveNodeAndPromoteChildren(Node);
+			}
 		}
 	}
 }
