Copying //UE4/Dev-Blueprints to Main (//UE4/Main)

diff --git a/Engine/Source/Developer/BlueprintCompilerCppBackend/Private/BlueprintCompilerCppBackend.cpp b/Engine/Source/Developer/BlueprintCompilerCppBackend/Private/BlueprintCompilerCppBackend.cpp
--- a/Engine/Source/Developer/BlueprintCompilerCppBackend/Private/BlueprintCompilerCppBackend.cpp
+++ b/Engine/Source/Developer/BlueprintCompilerCppBackend/Private/BlueprintCompilerCppBackend.cpp
@@ -444,10 +444,10 @@ FString FBlueprintCompilerCppBackend::EmitSwitchValueStatmentInner(FEmitterLocal
 		| EPropertyExportCPPFlags::CPPF_BlueprintCppBackend;
 
 	check(IndexTerm && IndexTerm->AssociatedVarProperty);
-	const FString IndexDeclaration = EmitterContext.ExportCppDeclaration(IndexTerm->AssociatedVarProperty, EExportedDeclaration::Local, CppTemplateTypeFlags, true);
+	const FString IndexDeclaration = EmitterContext.ExportCppDeclaration(IndexTerm->AssociatedVarProperty, EExportedDeclaration::Local, CppTemplateTypeFlags, FEmitterLocalContext::EPropertyNameInDeclaration::Skip);
 
 	check(DefaultValueTerm && DefaultValueTerm->AssociatedVarProperty);
-	const FString ValueDeclaration = EmitterContext.ExportCppDeclaration(DefaultValueTerm->AssociatedVarProperty, EExportedDeclaration::Local, CppTemplateTypeFlags, true);
+	const FString ValueDeclaration = EmitterContext.ExportCppDeclaration(DefaultValueTerm->AssociatedVarProperty, EExportedDeclaration::Local, CppTemplateTypeFlags, FEmitterLocalContext::EPropertyNameInDeclaration::Skip);
 
 	FString Result = FString::Printf(TEXT("TSwitchValue<%s, %s>(%s, %s, %d")
 		, *IndexDeclaration
@@ -692,6 +692,30 @@ FString FBlueprintCompilerCppBackend::EmitCallStatmentInner(FEmitterLocalContext
 		return TEXT("/*This function cannot be called from BP. See bIsValidFunction in UObject::CallFunction*/");
 	}
 
+	if (Statement.FunctionToCall->HasAllFunctionFlags(FUNC_Native))
+	{
+		// Cloned logic from: FScriptBytecodeWriter::EmitFunctionCall
+		// Array output parameters are cleared, in case the native function doesn't clear them before filling.
+		int32 NumParams = 0;
+		for (TFieldIterator<UProperty> PropIt(Statement.FunctionToCall); PropIt && (PropIt->PropertyFlags & CPF_Parm); ++PropIt)
+		{
+			UProperty* Param = *PropIt;
+			if (ensure(Param) && !Param->HasAnyPropertyFlags(CPF_ReturnParm))
+			{
+				const bool bShouldParameterBeCleared = Param->IsA<UArrayProperty>()
+					&& Param->HasAllPropertyFlags(CPF_Parm | CPF_OutParm)
+					&& !Param->HasAnyPropertyFlags(CPF_ReferenceParm | CPF_ConstParm | CPF_ReturnParm);
+				if (bShouldParameterBeCleared)
+				{
+					FBPTerminal* Term = Statement.RHS[NumParams];
+					const FString TermStr = ensure(Term != nullptr) ? TermToText(EmitterContext, Term, ENativizedTermUsage::UnspecifiedOrReference) : FString();
+					EmitterContext.AddLine(FString::Printf(TEXT("(%s).Reset();"), *TermStr));
+				}
+				NumParams++;
+			}
+		}
+	}
+
 	FString Result;
 	FString CloseCast;
 	TAutoPtr<FSetterExpressionBuilder> SetterExpression;
@@ -878,7 +902,7 @@ FString FBlueprintCompilerCppBackend::TermToText(FEmitterLocalContext& EmitterCo
 			const FString DefaultValueVariable = EmitterContext.GenerateUniqueLocalName();
 			const uint32 PropertyExportFlags = EPropertyExportCPPFlags::CPPF_CustomTypeName | EPropertyExportCPPFlags::CPPF_BlueprintCppBackend | EPropertyExportCPPFlags::CPPF_NoConst;
 			const FString CppType = Term->AssociatedVarProperty
-				? EmitterContext.ExportCppDeclaration(Term->AssociatedVarProperty, EExportedDeclaration::Local, PropertyExportFlags, true)
+				? EmitterContext.ExportCppDeclaration(Term->AssociatedVarProperty, EExportedDeclaration::Local, PropertyExportFlags, FEmitterLocalContext::EPropertyNameInDeclaration::Skip)
 				: FEmitHelper::PinTypeToNativeType(Term->Type);
 
 			const FString DefaultValueConstructor = (!Term->Type.bIsArray)
@@ -1005,29 +1029,21 @@ FString FBlueprintCompilerCppBackend::TermToText(FEmitterLocalContext& EmitterCo
 		}
 
 		const bool bNativeConstTemplateArg = Term->AssociatedVarProperty && Term->AssociatedVarProperty->HasMetaData(FName(TEXT("NativeConstTemplateArg")));
-		if ((bNativeConst || bNativeConstTemplateArg) && bIsAccessible && bGetter)
+		if (Term->Type.bIsArray && bNativeConstTemplateArg && bIsAccessible && bGetter)
 		{
-			if (Term->Type.bIsArray && bNativeConstTemplateArg)
-			{
-				FEdGraphPinType InnerType = Term->Type;
-				InnerType.bIsArray = false;
-				InnerType.bIsConst = false;
-				const FString CppType = FEmitHelper::PinTypeToNativeType(InnerType);
-				ResultPath = FString::Printf(TEXT("TArrayCaster<const %s>(%s).Get<%s>()"), *CppType, *ResultPath, *CppType);
-			}
-			else
-			{
-				const FString CppType = FEmitHelper::PinTypeToNativeType(Term->Type);
-				ResultPath = FString::Printf(TEXT("const_cast<%s>(%s)"), *CppType, *ResultPath);
-			}
+			FEdGraphPinType InnerType = Term->Type;
+			InnerType.bIsArray = false;
+			InnerType.bIsConst = false;
+			const FString CppType = FEmitHelper::PinTypeToNativeType(InnerType);
+			ResultPath = FString::Printf(TEXT("TArrayCaster<const %s>(%s).Get<%s>()"), *CppType, *ResultPath, *CppType);
 		}
-
-		FString Conditions;
-		if (bUseSafeContext)
+		else if ((bNativeConst || bNativeConstTemplateArg) && bIsAccessible && bGetter) //Why only getters?
 		{
-			Conditions = FSafeContextScopedEmmitter::ValidationChain(EmitterContext, Term->Context, *this);
+			const FString CppType = FEmitHelper::PinTypeToNativeType(Term->Type);
+			ResultPath = FString::Printf(TEXT("const_cast<%s>(%s)"), *CppType, *ResultPath);
 		}
 
+		const FString Conditions = bUseSafeContext ? FSafeContextScopedEmmitter::ValidationChain(EmitterContext, Term->Context, *this) : FString();
 		if (!Conditions.IsEmpty())
 		{
 			const FString DefaultValueVariable = GenerateDefaultLocalVariable(Term);
diff --git a/Engine/Source/Developer/BlueprintCompilerCppBackend/Private/BlueprintCompilerCppBackendBase.cpp b/Engine/Source/Developer/BlueprintCompilerCppBackend/Private/BlueprintCompilerCppBackendBase.cpp
--- a/Engine/Source/Developer/BlueprintCompilerCppBackend/Private/BlueprintCompilerCppBackendBase.cpp
+++ b/Engine/Source/Developer/BlueprintCompilerCppBackend/Private/BlueprintCompilerCppBackendBase.cpp
@@ -283,6 +283,22 @@ FString FBlueprintCompilerCppBackendBase::GenerateCodeFromClass(UClass* SourceCl
 
 	EmitStructProperties(EmitterContext, SourceClass);
 
+	{
+		IBlueprintCompilerCppBackendModule& BackEndModule = (IBlueprintCompilerCppBackendModule&)IBlueprintCompilerCppBackendModule::Get();
+		TSharedPtr<FNativizationSummary> NativizationSummary = BackEndModule.NativizationSummary();
+		if (NativizationSummary.IsValid())
+		{
+			for (TFieldIterator<UProperty> It(SourceClass, EFieldIteratorFlags::ExcludeSuper); It; ++It)
+			{
+				UProperty* Property = *It;
+				if (Property && Property->HasAllPropertyFlags(CPF_Transient | CPF_DuplicateTransient))
+				{
+					NativizationSummary->MemberVariablesFromGraph++;
+				}
+			}
+		}
+	}
+
 	TSharedPtr<FGatherConvertedClassDependencies> ParentDependencies;
 	// Emit function declarations and definitions (writes to header and body simultaneously)
 	if (!bIsInterface)
@@ -493,7 +509,7 @@ void FBlueprintCompilerCppBackendBase::ConstructFunction(FKismetFunctionContext&
 				if (FEmitHelper::PropertyForConstCast(Property))
 				{
 					const uint32 ExportFlags = EPropertyExportCPPFlags::CPPF_CustomTypeName | EPropertyExportCPPFlags::CPPF_BlueprintCppBackend | EPropertyExportCPPFlags::CPPF_NoConst | EPropertyExportCPPFlags::CPPF_NoRef;
-					const FString NoConstNoRefType = EmitterContext.ExportCppDeclaration(Property, EExportedDeclaration::Parameter, ExportFlags, true);
+					const FString NoConstNoRefType = EmitterContext.ExportCppDeclaration(Property, EExportedDeclaration::Parameter, ExportFlags, FEmitterLocalContext::EPropertyNameInDeclaration::Skip);
 					const FString TypeDefName = FString(TEXT("T")) + EmitterContext.GenerateUniqueLocalName();
 					EmitterContext.AddLine(FString::Printf(TEXT("typedef %s %s;"), *NoConstNoRefType, *TypeDefName));
 
@@ -699,7 +715,7 @@ FString FBlueprintCompilerCppBackendBase::GenerateArgList(const FEmitterLocalCon
 			ArgListStr += EmitterContext.ExportCppDeclaration(ArgProperty
 				, EExportedDeclaration::Parameter
 				, EPropertyExportCPPFlags::CPPF_CustomTypeName | EPropertyExportCPPFlags::CPPF_BlueprintCppBackend | EPropertyExportCPPFlags::CPPF_ArgumentOrReturnValue
-				, false
+				, FEmitterLocalContext::EPropertyNameInDeclaration::Regular
 				, NamePostFix);
 		}
 	}
@@ -718,7 +734,7 @@ FString FBlueprintCompilerCppBackendBase::GenerateReturnType(const FEmitterLocal
 			| EPropertyExportCPPFlags::CPPF_NoStaticArray
 			| EPropertyExportCPPFlags::CPPF_BlueprintCppBackend
 			| EPropertyExportCPPFlags::CPPF_ArgumentOrReturnValue;
-		return EmitterContext.ExportCppDeclaration(ReturnValue, EExportedDeclaration::Parameter, LocalExportCPPFlags, true);
+		return EmitterContext.ExportCppDeclaration(ReturnValue, EExportedDeclaration::Parameter, LocalExportCPPFlags, FEmitterLocalContext::EPropertyNameInDeclaration::Skip);
 	}
 	return TEXT("void");
 }
@@ -1001,7 +1017,10 @@ FString FBlueprintCompilerCppBackendBase::GenerateWrapperForClass(UClass* Source
 		//TODO: check if the property is really used?
 		const FString TypeDeclaration = Property->IsA<UMulticastDelegateProperty>()
 			? FString::Printf(TEXT("%s::%s"), *DelegatesClassName, *GenerateMulticastDelegateTypeName(CastChecked<UMulticastDelegateProperty>(Property)))
-			: EmitterContext.ExportCppDeclaration(Property, EExportedDeclaration::Parameter, EPropertyExportCPPFlags::CPPF_CustomTypeName | EPropertyExportCPPFlags::CPPF_BlueprintCppBackend | EPropertyExportCPPFlags::CPPF_NoRef, true);
+			: EmitterContext.ExportCppDeclaration(Property
+				, EExportedDeclaration::Parameter
+				, EPropertyExportCPPFlags::CPPF_CustomTypeName | EPropertyExportCPPFlags::CPPF_BlueprintCppBackend | EPropertyExportCPPFlags::CPPF_NoRef
+				, FEmitterLocalContext::EPropertyNameInDeclaration::Skip);
 		EmitterContext.Header.AddLine(FString::Printf(TEXT("FORCENOINLINE %s& GetRef__%s()"), *TypeDeclaration, *UnicodeToCPPIdentifier(Property->GetName(), false, nullptr)));
 		EmitterContext.Header.AddLine(TEXT("{"));
 		EmitterContext.Header.IncreaseIndent();
@@ -1077,7 +1096,11 @@ FString FBlueprintCompilerCppBackendBase::GenerateWrapperForClass(UClass* Source
 						ParamAsStructMember = TEXT("const ");
 					}
 				}
-				ParamAsStructMember += EmitterContext.ExportCppDeclaration(Property, EExportedDeclaration::Local, EPropertyExportCPPFlags::CPPF_CustomTypeName | EPropertyExportCPPFlags::CPPF_BlueprintCppBackend, false, ParamNameInStructPostfix);
+				ParamAsStructMember += EmitterContext.ExportCppDeclaration(Property
+					, EExportedDeclaration::Local
+					, EPropertyExportCPPFlags::CPPF_CustomTypeName | EPropertyExportCPPFlags::CPPF_BlueprintCppBackend
+					, FEmitterLocalContext::EPropertyNameInDeclaration::Regular
+					, ParamNameInStructPostfix);
 				FuncParameters.Emplace(ParamAsStructMember);
 				RawParameterList += UnicodeToCPPIdentifier(Property->GetName(), Property->HasAnyPropertyFlags(CPF_Deprecated), TEXT("bpp__"));
 			}
@@ -1170,7 +1193,7 @@ void FBlueprintCompilerCppBackendBase::EmitFileBeginning(const FString& CleanNam
 			{
 				// @TODO: Need to query if this asset will actually be converted
 
-				const FString Name = Field->GetName();
+				const FString Name = Field->GetPathName();
 				bool bAlreadyIncluded = false;
 				AlreadyIncluded.Add(Name, &bAlreadyIncluded);
 				if (!bAlreadyIncluded)
diff --git a/Engine/Source/Developer/BlueprintCompilerCppBackend/Private/BlueprintCompilerCppBackendGatherDependencies.cpp b/Engine/Source/Developer/BlueprintCompilerCppBackend/Private/BlueprintCompilerCppBackendGatherDependencies.cpp
--- a/Engine/Source/Developer/BlueprintCompilerCppBackend/Private/BlueprintCompilerCppBackendGatherDependencies.cpp
+++ b/Engine/Source/Developer/BlueprintCompilerCppBackend/Private/BlueprintCompilerCppBackendGatherDependencies.cpp
@@ -175,12 +175,29 @@ struct FFindHeadersToInclude : public FGatherConvertedClassDependenciesHelperBas
 			{
 				if (Graph)
 				{
-					TArray<UK2Node_EnumLiteral*> LiteralEnumNodes;
-					Graph->GetNodesOfClass<UK2Node_EnumLiteral>(LiteralEnumNodes);
-					for (UK2Node_EnumLiteral* LiteralEnumNode : LiteralEnumNodes)
+					TArray<UK2Node*> AllNodes;
+					Graph->GetNodesOfClass<UK2Node>(AllNodes);
+					for (UK2Node* K2Node : AllNodes)
 					{
-						UEnum* Enum = LiteralEnumNode ? LiteralEnumNode->Enum : nullptr;
-						IncludeTheHeaderInBody(Enum);
+						if (UK2Node_EnumLiteral* LiteralEnumNode = Cast<UK2Node_EnumLiteral>(K2Node))
+						{
+							UEnum* Enum = LiteralEnumNode ? LiteralEnumNode->Enum : nullptr;
+							IncludeTheHeaderInBody(Enum);
+						}
+						// HACK FOR LITERAL ENUMS:
+						else if(K2Node)
+						{
+							for (UEdGraphPin* Pin : K2Node->Pins)
+							{
+								if (Pin && (Pin->PinType.PinCategory == UEdGraphSchema_K2::PC_Byte))
+								{
+									if (UEnum* Enum = Cast<UEnum>(Pin->PinType.PinSubCategoryObject.Get()))
+									{
+										IncludeTheHeaderInBody(Enum);
+									}
+								}
+							}
+						}
 					}
 				}
 			}
@@ -405,6 +422,14 @@ void FGatherConvertedClassDependencies::DependenciesForHeader()
 		const bool bIsMemberVariable = OwnerProperty && (OwnerProperty->GetOuter() == OriginalStruct);
 		if (bIsParam || bIsMemberVariable)
 		{
+			if (auto AssetClassProperty = Cast<const UAssetClassProperty>(Property))
+			{
+				DeclareInHeader.Add(GetFirstNativeOrConvertedClass(AssetClassProperty->MetaClass));
+			}
+			if (auto ClassProperty = Cast<const UClassProperty>(Property))
+			{
+				DeclareInHeader.Add(GetFirstNativeOrConvertedClass(ClassProperty->MetaClass));
+			}
 			if (auto ObjectProperty = Cast<const UObjectPropertyBase>(Property))
 			{
 				DeclareInHeader.Add(GetFirstNativeOrConvertedClass(ObjectProperty->PropertyClass));
@@ -499,6 +524,25 @@ void FGatherConvertedClassDependencies::DependenciesForHeader()
 TSet<const UObject*> FGatherConvertedClassDependencies::AllDependencies() const
 {
 	TSet<const UObject*> All;
+
+	UBlueprintGeneratedClass* SuperClass = Cast<UBlueprintGeneratedClass>(OriginalStruct->GetSuperStruct());
+	if (SuperClass && WillClassBeConverted(SuperClass))
+	{
+		All.Add(SuperClass);
+	}
+
+	if (auto SourceClass = Cast<UClass>(OriginalStruct))
+	{
+		for (auto& ImplementedInterface : SourceClass->Interfaces)
+		{
+			UBlueprintGeneratedClass* InterfaceClass = Cast<UBlueprintGeneratedClass>(ImplementedInterface.Class);
+			if (InterfaceClass && WillClassBeConverted(InterfaceClass))
+			{
+				All.Add(InterfaceClass);
+			}
+		}
+	}
+
 	for (auto It : Assets)
 	{
 		All.Add(It);
diff --git a/Engine/Source/Developer/BlueprintCompilerCppBackend/Private/BlueprintCompilerCppBackendUMG.cpp b/Engine/Source/Developer/BlueprintCompilerCppBackend/Private/BlueprintCompilerCppBackendUMG.cpp
--- a/Engine/Source/Developer/BlueprintCompilerCppBackend/Private/BlueprintCompilerCppBackendUMG.cpp
+++ b/Engine/Source/Developer/BlueprintCompilerCppBackend/Private/BlueprintCompilerCppBackendUMG.cpp
@@ -55,7 +55,7 @@ void FBackendHelperUMG::EmitWidgetInitializationFunctions(FEmitterLocalContext&
 			const FString NativeName = InContext.GenerateUniqueLocalName();
 			
 			const uint32 CppTemplateTypeFlags = EPropertyExportCPPFlags::CPPF_CustomTypeName | EPropertyExportCPPFlags::CPPF_BlueprintCppBackend | EPropertyExportCPPFlags::CPPF_NoConst | EPropertyExportCPPFlags::CPPF_NoRef;
-			const FString Target = InContext.ExportCppDeclaration(InProperty, EExportedDeclaration::Local, CppTemplateTypeFlags, true);
+			const FString Target = InContext.ExportCppDeclaration(InProperty, EExportedDeclaration::Local, CppTemplateTypeFlags, FEmitterLocalContext::EPropertyNameInDeclaration::Skip);
 
 			InContext.AddLine(FString::Printf(TEXT("%s %s;"), *Target, *NativeName));
 			FEmitDefaultValueHelper::InnerGenerate(InContext, InProperty, NativeName, DataPtr, nullptr, true);
diff --git a/Engine/Source/Developer/BlueprintCompilerCppBackend/Private/BlueprintCompilerCppBackendUtils.cpp b/Engine/Source/Developer/BlueprintCompilerCppBackend/Private/BlueprintCompilerCppBackendUtils.cpp
--- a/Engine/Source/Developer/BlueprintCompilerCppBackend/Private/BlueprintCompilerCppBackendUtils.cpp
+++ b/Engine/Source/Developer/BlueprintCompilerCppBackend/Private/BlueprintCompilerCppBackendUtils.cpp
@@ -214,20 +214,28 @@ FString FEmitterLocalContext::ExportTextItem(const UProperty* Property, const vo
 	if (auto ArrayProperty = Cast<const UArrayProperty>(Property))
 	{
 		const FString ConstPrefix = Property->HasMetaData(TEXT("NativeConstTemplateArg")) ? TEXT("const ") : TEXT("");
-		const FString TypeText = ExportCppDeclaration(ArrayProperty, EExportedDeclaration::Parameter, LocalExportCPPFlags, true, FString(), ConstPrefix);
+		const FString TypeText = ExportCppDeclaration(ArrayProperty, EExportedDeclaration::Parameter, LocalExportCPPFlags, EPropertyNameInDeclaration::Skip, FString(), ConstPrefix);
 		return FString::Printf(TEXT("%s()"), *TypeText);
 	}
 	FString ValueStr;
 	Property->ExportTextItem(ValueStr, PropertyValue, PropertyValue, nullptr, EPropertyPortFlags::PPF_ExportCpp);
+	if (Property->IsA<UIntProperty>())
+	{
+		int32 Value = *(int32*)PropertyValue;
+		if (Value == (-2147483647 - 1)) // END OF RANGE
+		{
+			ValueStr = TEXT("(-2147483647 - 1)");
+		}
+	}
 	if (Property->IsA<UAssetObjectProperty>())
 	{
-		const FString TypeText = ExportCppDeclaration(Property, EExportedDeclaration::Parameter, LocalExportCPPFlags, true);
+		const FString TypeText = ExportCppDeclaration(Property, EExportedDeclaration::Parameter, LocalExportCPPFlags, EPropertyNameInDeclaration::Skip);
 		return FString::Printf(TEXT("%s(%s)"), *TypeText, *ValueStr);
 	}
 	return ValueStr;
 }
 
-FString FEmitterLocalContext::ExportCppDeclaration(const UProperty* Property, EExportedDeclaration::Type DeclarationType, uint32 InExportCPPFlags, bool bSkipParameterName, const FString& NamePostfix, const FString& TypePrefix) const
+FString FEmitterLocalContext::ExportCppDeclaration(const UProperty* Property, EExportedDeclaration::Type DeclarationType, uint32 InExportCPPFlags, FEmitterLocalContext::EPropertyNameInDeclaration ParameterName, const FString& NamePostfix, const FString& TypePrefix) const
 {
 	FString ActualCppType;
 	FString* ActualCppTypePtr = nullptr;
@@ -312,13 +320,14 @@ FString FEmitterLocalContext::ExportCppDeclaration(const UProperty* Property, EE
 	}
 
 	FStringOutputDevice Out;
-	const FString ActualNativeName = bSkipParameterName ? FString() : (FEmitHelper::GetCppName(Property) + NamePostfix);
+	const bool bSkipParameterName = (ParameterName == EPropertyNameInDeclaration::Skip);
+	const FString ActualNativeName = bSkipParameterName ? FString() : (FEmitHelper::GetCppName(Property, false, ParameterName == EPropertyNameInDeclaration::ForceConverted) + NamePostfix);
 	Property->ExportCppDeclaration(Out, DeclarationType, nullptr, InExportCPPFlags, bSkipParameterName, ActualCppTypePtr, ActualExtendedTypePtr, &ActualNativeName);
 	return FString(Out);
 
 }
 
-FString FEmitHelper::GetCppName(const UField* Field, bool bUInterface)
+FString FEmitHelper::GetCppName(const UField* Field, bool bUInterface, bool bForceParameterNameModification)
 {
 	check(Field);
 	auto AsClass = Cast<UClass>(Field);
@@ -344,26 +353,25 @@ FString FEmitHelper::GetCppName(const UField* Field, bool bUInterface)
 	}
 	else if (auto AsProperty = Cast<UProperty>(Field))
 	{
-		if (const UStruct* Owner = AsProperty->GetOwnerStruct())
+		const UStruct* Owner = AsProperty->GetOwnerStruct();
+		const bool bModifyName = ensure(Owner)
+			&& (Cast<UBlueprintGeneratedClass>(Owner) || !Owner->IsNative() || bForceParameterNameModification);
+		if (bModifyName)
 		{
-			if ((Cast<UBlueprintGeneratedClass>(Owner) ||
-				!Owner->IsNative()))
+			FString VarPrefix;
+
+			const bool bIsUberGraphVariable = Owner->IsA<UBlueprintGeneratedClass>() && AsProperty->HasAllPropertyFlags(CPF_Transient | CPF_DuplicateTransient);
+			if (bIsUberGraphVariable)
 			{
-				FString VarPrefix;
-				
-				const bool bIsUberGraphVariable = Owner->IsA<UBlueprintGeneratedClass>() && AsProperty->HasAllPropertyFlags(CPF_Transient | CPF_DuplicateTransient);
-				if (bIsUberGraphVariable)
-				{
-					int32 InheritenceLevel = GetInheritenceLevel(Owner);
-					VarPrefix = FString::Printf(TEXT("b%dl__"), InheritenceLevel);
-				}
-				else
-				{
-					const bool bIsParameter = AsProperty->HasAnyPropertyFlags(CPF_Parm);
-					VarPrefix = bIsParameter ? TEXT("bpp__") : TEXT("bpv__");
-				}
-				return ::UnicodeToCPPIdentifier(AsProperty->GetName(), AsProperty->HasAnyPropertyFlags(CPF_Deprecated), *VarPrefix);
+				int32 InheritenceLevel = GetInheritenceLevel(Owner);
+				VarPrefix = FString::Printf(TEXT("b%dl__"), InheritenceLevel);
 			}
+			else
+			{
+				const bool bIsParameter = AsProperty->HasAnyPropertyFlags(CPF_Parm);
+				VarPrefix = bIsParameter ? TEXT("bpp__") : TEXT("bpv__");
+			}
+			return ::UnicodeToCPPIdentifier(AsProperty->GetName(), AsProperty->HasAnyPropertyFlags(CPF_Deprecated), *VarPrefix);
 		}
 		return AsProperty->GetNameCPP();
 	}
@@ -394,7 +402,8 @@ int32 FEmitHelper::GetInheritenceLevel(const UStruct* Struct)
 
 bool FEmitHelper::PropertyForConstCast(const UProperty* Property)
 {
-	return Property && Property->HasAnyPropertyFlags(CPF_ConstParm);
+	return Property && (Property->HasAnyPropertyFlags(CPF_ConstParm) 
+			|| (Property->PassCPPArgsByRef() && !Property->HasAnyPropertyFlags(CPF_OutParm))); // See implementation in UProperty::ExportCppDeclaration
 }
 
 void FEmitHelper::ArrayToString(const TArray<FString>& Array, FString& OutString, const TCHAR* Separator)
@@ -487,18 +496,26 @@ FString FEmitHelper::HandleMetaData(const UField* Field, bool AddCategory, const
 	{
 		for (auto& Pair : *ValuesMap)
 		{
-			if (!MetaDataCanBeNative(Pair.Key, Field) || !IsMetaDataValid(Pair.Key, Pair.Value))
+			FName CurrentKey = Pair.Key;
+			FName NewKey = UMetaData::GetRemappedKeyName(CurrentKey);
+
+			if (NewKey != NAME_None)
+			{
+				CurrentKey = NewKey;
+			}
+
+			if (!MetaDataCanBeNative(CurrentKey, Field) || !IsMetaDataValid(CurrentKey, Pair.Value))
 			{
 				continue;
 			}
 			if (!Pair.Value.IsEmpty())
 			{
 				FString Value = Pair.Value.Replace(TEXT("\n"), TEXT("")).ReplaceCharWithEscapedChar();
-				MetaDataStrings.Emplace(FString::Printf(TEXT("%s=\"%s\""), *Pair.Key.ToString(), *Value));
+				MetaDataStrings.Emplace(FString::Printf(TEXT("%s=\"%s\""), *CurrentKey.ToString(), *Value));
 			}
 			else
 			{
-				MetaDataStrings.Emplace(Pair.Key.ToString());
+				MetaDataStrings.Emplace(CurrentKey.ToString());
 			}
 		}
 	}
@@ -759,7 +776,10 @@ int32 FEmitHelper::ParseDelegateDetails(FEmitterLocalContext& EmitterContext, UF
 	for (TFieldIterator<UProperty> PropIt(Signature); PropIt && (PropIt->PropertyFlags & CPF_Parm); ++PropIt)
 	{
 		Parameters += ", ";
-		Parameters += EmitterContext.ExportCppDeclaration(*PropIt, EExportedDeclaration::MacroParameter, EPropertyExportCPPFlags::CPPF_CustomTypeName | EPropertyExportCPPFlags::CPPF_BlueprintCppBackend);
+		Parameters += EmitterContext.ExportCppDeclaration(*PropIt
+			, EExportedDeclaration::MacroParameter
+			, EPropertyExportCPPFlags::CPPF_CustomTypeName | EPropertyExportCPPFlags::CPPF_BlueprintCppBackend
+			, FEmitterLocalContext::EPropertyNameInDeclaration::ForceConverted);
 		++ParameterNum;
 	}
 
@@ -803,7 +823,7 @@ void FEmitHelper::EmitSinglecastDelegateDeclarations(FEmitterLocalContext& Emitt
 			| EPropertyExportCPPFlags::CPPF_NoRef
 			| EPropertyExportCPPFlags::CPPF_NoStaticArray
 			| EPropertyExportCPPFlags::CPPF_BlueprintCppBackend;
-		const FString TypeName = EmitterContext.ExportCppDeclaration(It, EExportedDeclaration::Parameter, LocalExportCPPFlags, true);
+		const FString TypeName = EmitterContext.ExportCppDeclaration(It, EExportedDeclaration::Parameter, LocalExportCPPFlags, FEmitterLocalContext::EPropertyNameInDeclaration::Skip);
 		EmitSinglecastDelegateDeclarations_Inner(EmitterContext, It->SignatureFunction, TypeName);
 	}
 }
@@ -824,7 +844,7 @@ void FEmitHelper::EmitMulticastDelegateDeclarations(FEmitterLocalContext& Emitte
 			| EPropertyExportCPPFlags::CPPF_NoStaticArray
 			| EPropertyExportCPPFlags::CPPF_BlueprintCppBackend;
 		EmitterContext.Header.AddLine(FString::Printf(TEXT("UDELEGATE(%s)"), *FEmitHelper::HandleMetaData(Signature, false)));
-		const FString TypeName = EmitterContext.ExportCppDeclaration(*It, EExportedDeclaration::Parameter, LocalExportCPPFlags, true);
+		const FString TypeName = EmitterContext.ExportCppDeclaration(*It, EExportedDeclaration::Parameter, LocalExportCPPFlags, FEmitterLocalContext::EPropertyNameInDeclaration::Skip);
 		EmitterContext.Header.AddLine(FString::Printf(TEXT("DECLARE_DYNAMIC_MULTICAST_DELEGATE%s(%s%s);"), *ParamNumberStr, *TypeName, *Parameters));
 	}
 }
@@ -1059,35 +1079,38 @@ FString FEmitHelper::LiteralTerm(FEmitterLocalContext& EmitterContext, const FEd
 			}
 			return FString::Printf(TEXT("LoadClass<UClass>(nullptr, TEXT(\"%s\"), nullptr, 0, nullptr)"), *(LiteralObject->GetPathName().ReplaceCharWithEscapedChar()));
 		}
-		return FString(TEXT("nullptr"));
+		return FString(TEXT("((UClass*)nullptr)"));
 	}
 	else if((UEdGraphSchema_K2::PC_AssetClass == Type.PinCategory) || (UEdGraphSchema_K2::PC_Asset == Type.PinCategory))
 	{
+		UClass* MetaClass = Cast<UClass>(Type.PinSubCategoryObject.Get());
+		MetaClass = MetaClass ? MetaClass : UObject::StaticClass();
+		const FString ObjTypeStr = FEmitHelper::GetCppName(EmitterContext.GetFirstNativeOrConvertedClass(MetaClass));
+
 		if (LiteralObject)
 		{
 			const bool bAssetSubclassOf = (UEdGraphSchema_K2::PC_AssetClass == Type.PinCategory);
-			UClass* MetaClass = Cast<UClass>(Type.PinSubCategoryObject.Get());
-			MetaClass = MetaClass ? MetaClass : UObject::StaticClass();
 			return FString::Printf(TEXT("%s<%s>(FStringAssetReference(TEXT(\"%s\")))")
 				, bAssetSubclassOf ? TEXT("TAssetSubclassOf") : TEXT("TAssetPtr")
-				, *FEmitHelper::GetCppName(EmitterContext.GetFirstNativeOrConvertedClass(MetaClass))
+				, *ObjTypeStr
 				, *(LiteralObject->GetPathName().ReplaceCharWithEscapedChar()));
 		}
-		return FString(TEXT("nullptr"));
+		return FString::Printf(TEXT("((%s*)nullptr)"), *ObjTypeStr);
 	}
 	else if (UEdGraphSchema_K2::PC_Object == Type.PinCategory)
 	{
+		UClass* FoundClass = Cast<UClass>(Type.PinSubCategoryObject.Get());
+		UClass* ObjectClassToUse = FoundClass ? EmitterContext.GetFirstNativeOrConvertedClass(FoundClass) : UObject::StaticClass();
 		if (LiteralObject)
 		{
-			UClass* FoundClass = Cast<UClass>(Type.PinSubCategoryObject.Get());
-			UClass* ObjectClassToUse = FoundClass ? EmitterContext.GetFirstNativeOrConvertedClass(FoundClass) : UObject::StaticClass();
 			const FString MappedObject = EmitterContext.FindGloballyMappedObject(LiteralObject, ObjectClassToUse, true);
 			if (!MappedObject.IsEmpty())
 			{
 				return MappedObject;
 			}
 		}
-		return FString(TEXT("nullptr"));
+		const FString ObjTypeStr = FEmitHelper::GetCppName(EmitterContext.GetFirstNativeOrConvertedClass(ObjectClassToUse));
+		return FString::Printf(TEXT("((%s*)nullptr)"), *ObjTypeStr);
 	}
 	else if (UEdGraphSchema_K2::PC_Interface == Type.PinCategory)
 	{
@@ -1306,6 +1329,7 @@ bool FEmitHelper::GenerateAutomaticCast(FEmitterLocalContext& EmitterContext, co
 		UClass* LClass = EmitterContext.Dependencies.FindOriginalClass( Cast<UClass>(LType.PinSubCategoryObject.Get()) );
 		LClass = LClass ? EmitterContext.GetFirstNativeOrConvertedClass(LClass) : nullptr;
 		UClass* RClass = EmitterContext.Dependencies.FindOriginalClass( Cast<UClass>(RType.PinSubCategoryObject.Get()) );
+		RClass = RClass ? EmitterContext.GetFirstNativeOrConvertedClass(RClass) : nullptr;
 		if (!RType.bIsArray && LClass && RClass && (LType.bIsReference || bForceReference) && (LClass != RClass) && RClass->IsChildOf(LClass))
 		{
 			// when pointer is passed as reference, the type must be exactly the same
@@ -1541,7 +1565,7 @@ FString FEmitHelper::AccessInaccessibleProperty(FEmitterLocalContext& EmitterCon
 	const uint32 CppTemplateTypeFlags = EPropertyExportCPPFlags::CPPF_CustomTypeName
 		| EPropertyExportCPPFlags::CPPF_NoConst | EPropertyExportCPPFlags::CPPF_NoRef | EPropertyExportCPPFlags::CPPF_NoStaticArray
 		| EPropertyExportCPPFlags::CPPF_BlueprintCppBackend;
-	const FString TypeDeclaration = EmitterContext.ExportCppDeclaration(Property, EExportedDeclaration::Member, CppTemplateTypeFlags, true);
+	const FString TypeDeclaration = EmitterContext.ExportCppDeclaration(Property, EExportedDeclaration::Member, CppTemplateTypeFlags, FEmitterLocalContext::EPropertyNameInDeclaration::Skip);
 	
 	// Private Property Offset functions are generated only for private/protected properties - see PrivatePropertiesOffsetGetters in CodeGenerator.cpp
 	const bool bHasPPO = Property->HasAnyPropertyFlags(CPF_NativeAccessSpecifierPrivate | CPF_NativeAccessSpecifierProtected); 
@@ -1677,7 +1701,7 @@ FString FEmitHelper::AccessInaccessiblePropertyUsingOffset(FEmitterLocalContext&
 	const uint32 CppTemplateTypeFlags = EPropertyExportCPPFlags::CPPF_CustomTypeName
 		| EPropertyExportCPPFlags::CPPF_NoConst | EPropertyExportCPPFlags::CPPF_NoRef | EPropertyExportCPPFlags::CPPF_NoStaticArray
 		| EPropertyExportCPPFlags::CPPF_BlueprintCppBackend;
-	const FString TypeDeclaration = EmitterContext.ExportCppDeclaration(Property, EExportedDeclaration::Member, CppTemplateTypeFlags, true);
+	const FString TypeDeclaration = EmitterContext.ExportCppDeclaration(Property, EExportedDeclaration::Member, CppTemplateTypeFlags, FEmitterLocalContext::EPropertyNameInDeclaration::Skip);
 
 	const FString ArrayParams = (StaticArrayIdx != 0)
 		? FString::Printf(TEXT(", sizeof(%s), %d"), *TypeDeclaration, StaticArrayIdx)
diff --git a/Engine/Source/Developer/BlueprintCompilerCppBackend/Private/BlueprintCompilerCppBackendUtils.h b/Engine/Source/Developer/BlueprintCompilerCppBackend/Private/BlueprintCompilerCppBackendUtils.h
--- a/Engine/Source/Developer/BlueprintCompilerCppBackend/Private/BlueprintCompilerCppBackendUtils.h
+++ b/Engine/Source/Developer/BlueprintCompilerCppBackend/Private/BlueprintCompilerCppBackendUtils.h
@@ -162,7 +162,14 @@ public:
 	FString FindGloballyMappedObject(const UObject* Object, const UClass* ExpectedClass = nullptr, bool bLoadIfNotFound = false, bool bTryUsedAssetsList = true);
 
 	// Functions needed for Unconverted classes
-	FString ExportCppDeclaration(const UProperty* Property, EExportedDeclaration::Type DeclarationType, uint32 AdditionalExportCPPFlags, bool bSkipParameterName = false, const FString& NamePostfix = FString(), const FString& TypePrefix = FString()) const;
+	enum class EPropertyNameInDeclaration
+	{
+		Regular,
+		Skip,
+		ForceConverted,
+	};
+
+	FString ExportCppDeclaration(const UProperty* Property, EExportedDeclaration::Type DeclarationType, uint32 AdditionalExportCPPFlags, EPropertyNameInDeclaration ParameterName = EPropertyNameInDeclaration::Regular, const FString& NamePostfix = FString(), const FString& TypePrefix = FString()) const;
 	FString ExportTextItem(const UProperty* Property, const void* PropertyValue) const;
 
 	// AS FCodeText
@@ -188,7 +195,7 @@ private:
 struct FEmitHelper
 {
 	// bUInterface - use interface with "U" prefix, by default there is "I" prefix
-	static FString GetCppName(const UField* Field, bool bUInterface = false);
+	static FString GetCppName(const UField* Field, bool bUInterface = false, bool bForceParameterNameModification = false);
 
 	// returns an unique number for a structure in structures hierarchy
 	static int32 GetInheritenceLevel(const UStruct* Struct);
diff --git a/Engine/Source/Developer/BlueprintCompilerCppBackend/Private/BlueprintCompilerCppBackendValueHelper.cpp b/Engine/Source/Developer/BlueprintCompilerCppBackend/Private/BlueprintCompilerCppBackendValueHelper.cpp
--- a/Engine/Source/Developer/BlueprintCompilerCppBackend/Private/BlueprintCompilerCppBackendValueHelper.cpp
+++ b/Engine/Source/Developer/BlueprintCompilerCppBackend/Private/BlueprintCompilerCppBackendValueHelper.cpp
@@ -792,7 +792,7 @@ void FEmitDefaultValueHelper::AddStaticFunctionsForDependencies(FEmitterLocalCon
 	// HELPERS
 	auto SourceClass = Context.GetCurrentlyGeneratedClass();
 	auto OriginalClass = Context.Dependencies.FindOriginalClass(SourceClass);
-	const FString CppClassName = FEmitHelper::GetCppName(OriginalClass);
+	const FString CppClassName = FEmitHelper::GetCppName(SourceClass);
 
 	auto CreateAssetToLoadString = [&](const UObject* AssetObj, const FString& LongPathVariable) -> FString
 	{
@@ -931,7 +931,7 @@ void FEmitDefaultValueHelper::AddRegisterHelper(FEmitterLocalContext& Context)
 {
 	auto SourceClass = Context.GetCurrentlyGeneratedClass();
 	auto OriginalClass = Context.Dependencies.FindOriginalClass(SourceClass);
-	const FString CppClassName = FEmitHelper::GetCppName(OriginalClass);
+	const FString CppClassName = FEmitHelper::GetCppName(SourceClass);
 
 	const FString RegisterHelperName = FString::Printf(TEXT("FRegisterHelper__%s"), *CppClassName);
 	Context.AddLine(FString::Printf(TEXT("struct %s"), *RegisterHelperName));
diff --git a/Engine/Source/Developer/BlueprintCompilerCppBackend/Public/IBlueprintCompilerCppBackendModule.h b/Engine/Source/Developer/BlueprintCompilerCppBackend/Public/IBlueprintCompilerCppBackendModule.h
--- a/Engine/Source/Developer/BlueprintCompilerCppBackend/Public/IBlueprintCompilerCppBackendModule.h
+++ b/Engine/Source/Developer/BlueprintCompilerCppBackend/Public/IBlueprintCompilerCppBackendModule.h
@@ -23,6 +23,10 @@ struct FNativizationSummary
 	};
 
 	TMap<FStringAssetReference, FAnimBlueprintDetails> AnimBlueprintStat;
+
+	int32 MemberVariablesFromGraph;
+
+	FNativizationSummary() : MemberVariablesFromGraph(0) {}
 };
 
 /**
diff --git a/Engine/Source/Developer/BlueprintNativeCodeGen/Private/BlueprintNativeCodeGenModule.cpp b/Engine/Source/Developer/BlueprintNativeCodeGen/Private/BlueprintNativeCodeGenModule.cpp
--- a/Engine/Source/Developer/BlueprintNativeCodeGen/Private/BlueprintNativeCodeGenModule.cpp
+++ b/Engine/Source/Developer/BlueprintNativeCodeGen/Private/BlueprintNativeCodeGenModule.cpp
@@ -69,9 +69,15 @@ private:
 
 	TMap< FString, TUniquePtr<FBlueprintNativeCodeGenManifest> > Manifests;
 
+	// Children of these classes won't be nativized
 	TArray<FString> ExcludedAssetTypes;
+	// Eg: +ExcludedBlueprintTypes=/Script/Engine.AnimBlueprint
 	TArray<TAssetSubclassOf<UBlueprint>> ExcludedBlueprintTypes;
+	// Individually excluded assets
 	TSet<FStringAssetReference> ExcludedAssets;
+	// Excluded folders. It excludes only BPGCs, enums and structures are still converted.
+	TArray<FString> ExcludedFolderPaths;
+
 	TArray<FString> TargetPlatformNames;
 
 	// A stub-wrapper must be generated only if the BP is really accessed/required by some other generated code.
@@ -104,6 +110,8 @@ void FBlueprintNativeCodeGenModule::ReadConfig()
 	{
 		ExcludedAssets.Add(FStringAssetReference(Path));
 	}
+
+	GConfig->GetArray(TEXT("BlueprintNativizationSettings"), TEXT("ExcludedFolderPaths"), ExcludedFolderPaths, GEditorIni);
 }
 
 void FBlueprintNativeCodeGenModule::MarkUnconvertedBlueprintAsNecessary(TAssetPtr<UBlueprint> BPPtr)
@@ -333,6 +341,7 @@ void FBlueprintNativeCodeGenModule::GenerateFullyConvertedClasses()
 				, Iter.Value.VariableUsage);
 		}
 	}
+	UE_LOG(LogBlueprintCodeGen, Display, TEXT("Nativization Summary - Shared Variables From Graph: %d"), NativizationSummary->MemberVariablesFromGraph);
 }
 
 FBlueprintNativeCodeGenManifest& FBlueprintNativeCodeGenModule::GetManifest(const TCHAR* PlatformName)
@@ -652,72 +661,150 @@ EReplacementResult FBlueprintNativeCodeGenModule::IsTargetedForReplacement(const
 
 EReplacementResult FBlueprintNativeCodeGenModule::IsTargetedForReplacement(const UObject* Object) const
 {
-	if (Object == nullptr)
+	const UStruct* const Struct = Cast<UStruct>(Object);
+	const UEnum* const Enum = Cast<UEnum>(Object);
+	const UClass* const BlueprintClass = Cast<UClass>(Struct);
+	const UBlueprint* const Blueprint = BlueprintClass ? Cast<UBlueprint>(BlueprintClass->ClassGeneratedBy) : nullptr;
+
+	auto ObjectIsNotReplacedAtAll = [&]() -> bool
 	{
-		return EReplacementResult::DontReplace;
-	}
+		if (Object == nullptr)
+		{
+			return true;
+		}
 
-	const UStruct* Struct = Cast<UStruct>(Object);
-	const UEnum* Enum = Cast<UEnum>(Object);
+		if (Struct == nullptr && Enum == nullptr)
+		{
+			return true;
+		}
 
-	if (Struct == nullptr && Enum == nullptr)
+		// EDITOR ON DEVELOPMENT OBJECT
+		{
+			auto IsDeveloperObject = [&](const UObject* Obj) -> bool
+			{
+				auto IsObjectFromDeveloperPackage = [](const UObject* InObj) -> bool
+				{
+					return InObj && InObj->GetOutermost()->HasAllPackagesFlags(PKG_Developer);
+				};
+
+				if (Obj)
+				{
+					if (IsObjectFromDeveloperPackage(Obj))
+					{
+						return true;
+					}
+					const UStruct* StructToTest = Obj->IsA<UStruct>() ? CastChecked<const UStruct>(Obj) : Obj->GetClass();
+					for (; StructToTest; StructToTest = StructToTest->GetSuperStruct())
+					{
+						if (IsObjectFromDeveloperPackage(StructToTest))
+						{
+							return true;
+						}
+					}
+				}
+				return false;
+			};
+			if (Object && (IsEditorOnlyObject(Object) || IsDeveloperObject(Object)))
+			{
+				UE_LOG(LogBlueprintCodeGen, Warning, TEXT("Object %s depends on Editor or Development stuff. It shouldn't be cooked."), *GetPathNameSafe(Object));
+				return true;
+			}
+		}
+		// DATA ONLY BP
+		{
+			static const FBoolConfigValueHelper DontNativizeDataOnlyBP(TEXT("BlueprintNativizationSettings"), TEXT("bDontNativizeDataOnlyBP"));
+			if (DontNativizeDataOnlyBP && Blueprint && FBlueprintEditorUtils::IsDataOnlyBlueprint(Blueprint))
+			{
+				return true;
+			}
+		}
+		return false;
+	};
+	if (ObjectIsNotReplacedAtAll())
 	{
 		return EReplacementResult::DontReplace;
 	}
 
-	EReplacementResult Result = EReplacementResult::ReplaceCompletely;
-	if (const UClass* BlueprintClass = Cast<UClass>(Struct))
+	auto ObjectGenratesOnlyStub = [&]() -> bool
 	{
-		if (UBlueprint* Blueprint = Cast<UBlueprint>(BlueprintClass->ClassGeneratedBy))
+		// ExcludedFolderPaths
+		if (BlueprintClass)
 		{
-			static const FBoolConfigValueHelper NativizeAnimBPOnlyWhenNonReducibleFuncitons(TEXT("BlueprintNativizationSettings"), TEXT("bNativizeAnimBPOnlyWhenNonReducibleFuncitons"));
-			if (NativizeAnimBPOnlyWhenNonReducibleFuncitons)
+			const FString ObjPathName = Object->GetPathName();
+			for (const FString& ExcludedPath : ExcludedFolderPaths)
 			{
-				if (UAnimBlueprint* AnimBlueprint = Cast<UAnimBlueprint>(Blueprint))
+				if (ObjPathName.StartsWith(ExcludedPath))
 				{
-					ensure(AnimBlueprint->bHasBeenRegenerated);
-					if (AnimBlueprint->bHasAnyNonReducibleFunction == UBlueprint::EIsBPNonReducible::No)
-					{
-						UE_LOG(LogBlueprintCodeGen, Log, TEXT("AnimBP %s without non-reducible functions is excluded from nativization"), *GetPathNameSafe(Blueprint));
-						Result = EReplacementResult::GenerateStub;
-					}
+					return true;
 				}
 			}
+		}
 
-			const EBlueprintType UnconvertableBlueprintTypes[] = {
-				//BPTYPE_Const,		// WTF is a "const" Blueprint?
-				BPTYPE_MacroLibrary,
-				BPTYPE_LevelScript,
-			};
+		// ExcludedAssetTypes
+		{
+			// we can't use FindObject, because we may be converting a type while saving
+			if (Enum && ExcludedAssetTypes.Find(Enum->GetPathName()) != INDEX_NONE)
+			{
+				return true;
+			}
 
-			EBlueprintType BlueprintType = Blueprint->BlueprintType;
-			for (int32 TypeIndex = 0; TypeIndex < ARRAY_COUNT(UnconvertableBlueprintTypes); ++TypeIndex)
+			const UStruct* LocStruct = Struct;
+			while (LocStruct)
 			{
-				if (BlueprintType == UnconvertableBlueprintTypes[TypeIndex])
+				if (ExcludedAssetTypes.Find(LocStruct->GetPathName()) != INDEX_NONE)
 				{
-					Result = EReplacementResult::GenerateStub;
+					return true;
 				}
+				LocStruct = LocStruct->GetSuperStruct();
 			}
+		}
 
-			static const FBoolConfigValueHelper DontNativizeDataOnlyBP(TEXT("BlueprintNativizationSettings"), TEXT("bDontNativizeDataOnlyBP"));
-			if (DontNativizeDataOnlyBP)
+		// ExcludedAssets
+		{
+			if (ExcludedAssets.Contains(Object->GetOutermost()))
+			{
+				return true;
+			}
+		}
+
+
+		if (Blueprint && BlueprintClass)
+		{
+			// Reducible AnimBP
 			{
-				if (FBlueprintEditorUtils::IsDataOnlyBlueprint(Blueprint))
+				static const FBoolConfigValueHelper NativizeAnimBPOnlyWhenNonReducibleFuncitons(TEXT("BlueprintNativizationSettings"), TEXT("bNativizeAnimBPOnlyWhenNonReducibleFuncitons"));
+				if (NativizeAnimBPOnlyWhenNonReducibleFuncitons)
 				{
-					return EReplacementResult::DontReplace;
+					if (const UAnimBlueprint* AnimBlueprint = Cast<UAnimBlueprint>(Blueprint))
+					{
+						ensure(AnimBlueprint->bHasBeenRegenerated);
+						if (AnimBlueprint->bHasAnyNonReducibleFunction == UBlueprint::EIsBPNonReducible::No)
+						{
+							UE_LOG(LogBlueprintCodeGen, Log, TEXT("AnimBP %s without non-reducible functions is excluded from nativization"), *GetPathNameSafe(Blueprint));
+							return true;
+						}
+					}
 				}
 			}
 
-			for (UBlueprintGeneratedClass* ParentClassIt = Cast<UBlueprintGeneratedClass>(BlueprintClass->GetSuperClass())
-				; ParentClassIt; ParentClassIt = Cast<UBlueprintGeneratedClass>(ParentClassIt->GetSuperClass()))
+			// Unconvertable Blueprint
 			{
-				EReplacementResult ParentResult = IsTargetedForReplacement(ParentClassIt);
-				if (ParentResult != EReplacementResult::ReplaceCompletely)
+				const EBlueprintType UnconvertableBlueprintTypes[] = {
+					//BPTYPE_Const,		// What is a "const" Blueprint?
+					BPTYPE_MacroLibrary,
+					BPTYPE_LevelScript,
+				};
+				const EBlueprintType BlueprintType = Blueprint->BlueprintType;
+				for (int32 TypeIndex = 0; TypeIndex < ARRAY_COUNT(UnconvertableBlueprintTypes); ++TypeIndex)
 				{
-					Result = EReplacementResult::GenerateStub;
+					if (BlueprintType == UnconvertableBlueprintTypes[TypeIndex])
+					{
+						return true;
+					}
 				}
 			}
 
+			// ExcludedBlueprintTypes
 			for (TAssetSubclassOf<UBlueprint> ExcludedBlueprintTypeAsset : ExcludedBlueprintTypes)
 			{
 				UClass* ExcludedBPClass = ExcludedBlueprintTypeAsset.Get();
@@ -727,65 +814,42 @@ EReplacementResult FBlueprintNativeCodeGenModule::IsTargetedForReplacement(const
 				}
 				if (ExcludedBPClass && Blueprint->IsA(ExcludedBPClass))
 				{
-					Result = EReplacementResult::GenerateStub;
+					return true;
 				}
 			}
-		}
-	}
 
-	auto IsObjectFromDeveloperPackage = [](const UObject* Obj) -> bool
-	{
-		return Obj && Obj->GetOutermost()->HasAllPackagesFlags(PKG_Developer);
-	};
+			const UProjectPackagingSettings* const PackagingSettings = GetDefault<UProjectPackagingSettings>();
+			const bool bNativizeOnlySelectedBPs = PackagingSettings && PackagingSettings->bNativizeOnlySelectedBlueprints;
 
-	auto IsDeveloperObject = [&](const UObject* Obj) -> bool
-	{
-		if (Obj)
-		{
-			if (IsObjectFromDeveloperPackage(Obj))
+			// Blueprint is not selected
+			if (bNativizeOnlySelectedBPs && !Blueprint->bNativize)
 			{
 				return true;
 			}
-			const UStruct* StructToTest = Obj->IsA<UStruct>() ? CastChecked<const UStruct>(Obj) : Obj->GetClass();
-			for (; StructToTest; StructToTest = StructToTest->GetSuperStruct())
+
+			// Parent Class in not converted
+			for (const UBlueprintGeneratedClass* ParentClassIt = Cast<UBlueprintGeneratedClass>(BlueprintClass->GetSuperClass())
+				; ParentClassIt; ParentClassIt = Cast<UBlueprintGeneratedClass>(ParentClassIt->GetSuperClass()))
 			{
-				if (IsObjectFromDeveloperPackage(StructToTest))
+				const EReplacementResult ParentResult = IsTargetedForReplacement(ParentClassIt);
+				if (ParentResult != EReplacementResult::ReplaceCompletely)
 				{
+					if (bNativizeOnlySelectedBPs)
+					{
+						UE_LOG(LogBlueprintCodeGen, Error, TEXT("BP %s is selected for nativization, but its parent class %s is not nativized."), *GetPathNameSafe(Blueprint), *GetPathNameSafe(ParentClassIt));
+					}
 					return true;
 				}
 			}
 		}
 		return false;
 	};
-
-	if (Object && (IsEditorOnlyObject(Object) || IsDeveloperObject(Object)))
-	{
-		UE_LOG(LogBlueprintCodeGen, Warning, TEXT("Object %s depends on Editor or Development stuff. It shouldn't be cooked."), *GetPathNameSafe(Object));
-		return EReplacementResult::DontReplace;
-	}
-
-	// check blacklists:
-	// we can't use FindObject, because we may be converting a type while saving
-	if (Enum && ExcludedAssetTypes.Find(Enum->GetPathName()) != INDEX_NONE)
-	{
-		Result = EReplacementResult::GenerateStub;
-	}
-
-	while (Struct)
-	{
-		if (ExcludedAssetTypes.Find(Struct->GetPathName()) != INDEX_NONE)
-		{
-			Result = EReplacementResult::GenerateStub;
-		}
-		Struct = Struct->GetSuperStruct();
-	}
-
-	if (ExcludedAssets.Contains(Object->GetOutermost()))
+	if (ObjectGenratesOnlyStub())
 	{
-		Result = EReplacementResult::GenerateStub;
+		return EReplacementResult::GenerateStub;
 	}
 
-	return Result;
+	return EReplacementResult::ReplaceCompletely;
 }
 
 IMPLEMENT_MODULE(FBlueprintNativeCodeGenModule, BlueprintNativeCodeGen);
\ No newline at end of file
diff --git a/Engine/Source/Developer/BlueprintNativeCodeGen/Private/NativeCodeGenerationTool.cpp b/Engine/Source/Developer/BlueprintNativeCodeGen/Private/NativeCodeGenerationTool.cpp
--- a/Engine/Source/Developer/BlueprintNativeCodeGen/Private/NativeCodeGenerationTool.cpp
+++ b/Engine/Source/Developer/BlueprintNativeCodeGen/Private/NativeCodeGenerationTool.cpp
@@ -131,7 +131,8 @@ struct FGeneratedCodeData
 		IBlueprintCompilerCppBackendModule& CodeGenBackend = (IBlueprintCompilerCppBackendModule&)IBlueprintCompilerCppBackendModule::Get();
 
 		TArray<FString> CreatedFiles;
-		for(auto Obj : DependentObjects)
+		//for(auto Obj : DependentObjects)
+		UObject* Obj = Blueprint->GeneratedClass;
 		{
 			SlowTask.EnterProgressFrame();
 
diff --git a/Engine/Source/Developer/BlueprintProfiler/Private/BlueprintProfiler.cpp b/Engine/Source/Developer/BlueprintProfiler/Private/BlueprintProfiler.cpp
--- a/Engine/Source/Developer/BlueprintProfiler/Private/BlueprintProfiler.cpp
+++ b/Engine/Source/Developer/BlueprintProfiler/Private/BlueprintProfiler.cpp
@@ -341,8 +341,8 @@ void FBlueprintProfiler::ProcessEventProfilingData()
 						{
 							EventToProcess = (MakeShareable(new FScriptEventPlayback(EventRangeToProcess.BlueprintContext, EventRangeToProcess.InstanceName)));
 						}
-						check(EventToProcess.IsValid());
-						if (EventToProcess->Process(InstrumentationEventQueue, EventRangeToProcess.StartIdx, EventRangeToProcess.StopIdx))
+						// Checking the event to process so we can continue problem discovery while we look into the latent id issue.
+						if (EventToProcess.IsValid() && EventToProcess->Process(InstrumentationEventQueue, EventRangeToProcess.StartIdx, EventRangeToProcess.StopIdx))
 						{
 							const int32 NumEventsToRemove = (EventRangeToProcess.StopIdx - EventRangeToProcess.StartIdx) + 1;
 							InstrumentationEventQueue.RemoveAt(EventRangeToProcess.StartIdx, NumEventsToRemove, false);
@@ -363,10 +363,8 @@ void FBlueprintProfiler::ProcessEventProfilingData()
 			}
 		}
 	}
-	// For the time being, we should have processed all events in the queue.
-	// In the future we will potentially only process a few events at a time, and this check should be removed.
-	// Removed this check after finding an issue with it switching out maps.
-//	check(InstrumentationEventQueue.Num() == 0);
+	// We need to ensure this doesn't happen. but for now kill anything remaining in the queue.
+	InstrumentationEventQueue.SetNum(0);
 	// Update all active contexts if the display settings changed.
 	UBlueprintProfilerSettings* Settings = GetMutableDefault<UBlueprintProfilerSettings>();
 	if (Settings->GetPerformanceThresholdsModified())
diff --git a/Engine/Source/Developer/BlueprintProfiler/Private/ScriptInstrumentationPlayback.cpp b/Engine/Source/Developer/BlueprintProfiler/Private/ScriptInstrumentationPlayback.cpp
--- a/Engine/Source/Developer/BlueprintProfiler/Private/ScriptInstrumentationPlayback.cpp
+++ b/Engine/Source/Developer/BlueprintProfiler/Private/ScriptInstrumentationPlayback.cpp
@@ -575,9 +575,6 @@ void FBlueprintFunctionContext::InitialiseContextFromGraph(TSharedPtr<FBlueprint
 						BlueprintContextIn->AddEventNode(AsShared(), EventExecNode);
 					}
 				}
-				// Create any compiler generated events
-				FBlueprintDebugData& DebugData = BlueprintClass.Get()->GetDebugData();
-				CreateDelegatePinEvents(BlueprintContextIn, DebugData.GetCompilerGeneratedEvents());
 			}
 			else
 			{
@@ -677,12 +674,17 @@ void FBlueprintFunctionContext::MapFunction()
 			}
 		}
 		// Check for Cyclic Linkage
-		TSet<TSharedPtr<FScriptExecutionNode>> Filter;
+		TArray<TSharedPtr<FScriptExecutionNode>> Filter;
 		DetectCyclicLinks(EntryPoint, Filter);
 	}
+	if (IsUbergraphFunction())
+	{
+		// Create any compiler generated events
+		CreateDelegatePinEvents();
+	}
 }
 
-bool FBlueprintFunctionContext::DetectCyclicLinks(TSharedPtr<FScriptExecutionNode> ExecNode, TSet<TSharedPtr<FScriptExecutionNode>>& Filter)
+bool FBlueprintFunctionContext::DetectCyclicLinks(TSharedPtr<FScriptExecutionNode> ExecNode, TArray<TSharedPtr<FScriptExecutionNode>>& Filter)
 {
 	if (ExecNode->HasFlags(EScriptExecutionNodeFlags::PureStats))
 	{
@@ -698,26 +700,30 @@ bool FBlueprintFunctionContext::DetectCyclicLinks(TSharedPtr<FScriptExecutionNod
 		{
 			Filter.Add(ExecNode);
 		}
-		for (TSharedPtr<FScriptExecutionNode>& Child : ExecNode->GetChildNodes())
+		if (!ExecNode->IsTunnelEntry())
 		{
-			if (DetectCyclicLinks(Child, Filter))
+			for (TSharedPtr<FScriptExecutionNode>& Child : ExecNode->GetChildNodes())
 			{
-				// Replace child with cyclic link node.
-				FScriptExecNodeParams CycleLinkParams;
-				CycleLinkParams.SampleFrequency = 1;
-				const FString NodeName = FString::Printf(TEXT("CyclicLinkTo_%i_%s"), ExecutionNodes.Num(), *Child->GetName().ToString());
-				CycleLinkParams.NodeName = FName(*NodeName);
-				CycleLinkParams.ObservedObject = Child->GetObservedObject();
-				CycleLinkParams.DisplayName = Child->GetDisplayName();
-				CycleLinkParams.Tooltip = LOCTEXT("CyclicLink_ToolTip", "Cyclic Link");
-				CycleLinkParams.NodeFlags = EScriptExecutionNodeFlags::CyclicLinkage|Child->GetFlags();
-				CycleLinkParams.IconColor = Child->GetIconColor();
-				CycleLinkParams.IconColor.A = 0.15f;
-				const FSlateBrush* LinkIcon = Child->GetIcon();
-				CycleLinkParams.Icon = const_cast<FSlateBrush*>(LinkIcon);
-				Child = CreateExecutionNode(CycleLinkParams);
+				if (DetectCyclicLinks(Child, Filter))
+				{
+					// Replace child with cyclic link node.
+					FScriptExecNodeParams CycleLinkParams;
+					CycleLinkParams.SampleFrequency = 1;
+					const FString NodeName = FString::Printf(TEXT("CyclicLinkTo_%i_%s"), ExecutionNodes.Num(), *Child->GetName().ToString());
+					CycleLinkParams.NodeName = FName(*NodeName);
+					CycleLinkParams.ObservedObject = Child->GetObservedObject();
+					CycleLinkParams.DisplayName = Child->GetDisplayName();
+					CycleLinkParams.Tooltip = LOCTEXT("CyclicLink_ToolTip", "Cyclic Link");
+					CycleLinkParams.NodeFlags = EScriptExecutionNodeFlags::CyclicLinkage|Child->GetFlags();
+					CycleLinkParams.IconColor = Child->GetIconColor();
+					CycleLinkParams.IconColor.A = 0.15f;
+					const FSlateBrush* LinkIcon = Child->GetIcon();
+					CycleLinkParams.Icon = const_cast<FSlateBrush*>(LinkIcon);
+					Child = CreateExecutionNode(CycleLinkParams);
+				}
 			}
 		}
+		const int32 FilterWaterMark = Filter.Num();
 		for (auto LinkedNode : ExecNode->GetLinkedNodes())
 		{
 			TSharedPtr<FScriptExecutionNode> LinkedExecNode = LinkedNode.Value;
@@ -733,8 +739,8 @@ bool FBlueprintFunctionContext::DetectCyclicLinks(TSharedPtr<FScriptExecutionNod
 					}
 				}
 			}
-			TSet<TSharedPtr<FScriptExecutionNode>> BranchFilter(Filter);
-			if (DetectCyclicLinks(LinkedExecNode, BranchFilter))
+			Filter.SetNum(FilterWaterMark);
+			if (DetectCyclicLinks(LinkedExecNode, Filter))
 			{
 				// Break links and flag cycle linkage.
 				FScriptExecNodeParams CycleLinkParams;
@@ -765,23 +771,31 @@ void FBlueprintFunctionContext::AddChildFunctionContext(const FName FunctionName
 	}
 }
 
-void FBlueprintFunctionContext::CreateDelegatePinEvents(TSharedPtr<FBlueprintExecutionContext> BlueprintContextIn, const TMap<FName, FEdGraphPinReference>& PinEvents)
+void FBlueprintFunctionContext::CreateDelegatePinEvents()
 {
 	struct FPinDelegateDesc
 	{
-		FPinDelegateDesc(const FName EventNameIn, UEdGraphPin* DelegatePinIn)
+		FPinDelegateDesc(const FName EventNameIn, UEdGraphPin* DelegatePinIn, const int32 ScriptOffsetIn, FNodeExecutionContext& ContextIn)
 			: EventName(EventNameIn)
 			, DelegatePin(DelegatePinIn)
+			, ScriptOffset(ScriptOffsetIn)
+			, Context(ContextIn)
 		{
 		}
 
 		FName EventName;
 		UEdGraphPin* DelegatePin;
+		const int32 ScriptOffset;
+		FNodeExecutionContext Context;
 	};
+
+	FBlueprintDebugData& DebugData = BlueprintClass.Get()->GetDebugData();
+	const TMap<FName, FEdGraphPinReference>& PinEvents = DebugData.GetCompilerGeneratedEvents();
 	if (PinEvents.Num())
 	{
 		TMap<const UEdGraphNode*, TArray<FPinDelegateDesc>> NodeEventDescs;
 		TSharedPtr<FBlueprintFunctionContext> FunctionContext = AsShared();
+		TSharedPtr<FBlueprintExecutionContext> ParentBlueprintContext = BlueprintContext.Pin();
 		// Build event contexts per node
 		for (auto PinEvent : PinEvents)
 		{
@@ -789,67 +803,114 @@ void FBlueprintFunctionContext::CreateDelegatePinEvents(TSharedPtr<FBlueprintExe
 			{
 				const UEdGraphNode* OwningNode = DelegatePin->GetOwningNode();
 				TArray<FPinDelegateDesc>& Events = NodeEventDescs.FindOrAdd(OwningNode);
-				const FName ScopedEventName = GetScopedEventName(PinEvent.Key);
-				Events.Add(FPinDelegateDesc(ScopedEventName, DelegatePin));
+				const int32 ScriptOffset = GetCodeLocationFromPin(DelegatePin);
+				FNodeExecutionContext& NodeContext = GetNodeExecutionContext(ScriptOffset);
+				if (NodeContext.IsValid())
+				{
+					const FName ScopedEventName = NodeContext.ProfilerContext->GetScopedEventName(PinEvent.Key);
+					Events.Add(FPinDelegateDesc(ScopedEventName, DelegatePin, ScriptOffset, NodeContext));
+				}
+				else
+				{
+					NodeContext.ProfilerContext = AsShared();
+					NodeContext.GraphNode = OwningNode;
+					const FName ScopedEventName = GetScopedEventName(PinEvent.Key);
+					Events.Add(FPinDelegateDesc(ScopedEventName, DelegatePin, ScriptOffset, NodeContext));
+				}
 			}
 		}
 		// Generate the event exec nodes
 		for (auto NodeEvents : NodeEventDescs)
 		{
-			// Check if this node requires an event node creating.
-			bool bCreateEventNode = true;
-			for (auto Pin : NodeEvents.Key->Pins)
-			{
-				if (Pin->Direction == EGPD_Input && Pin->PinType.PinCategory == UEdGraphSchema_K2::PC_Exec)
-				{
-					bCreateEventNode = false;
-					break;
-				}
-			}
 			TSharedPtr<FScriptExecutionNode> EventExecNode;
-			if (bCreateEventNode)
-			{
-				// Setup the basic exec node params.
-				FScriptExecNodeParams EventParams;
-				EventParams.SampleFrequency = 1;
-				EventParams.NodeName = FName(*FString::Printf(TEXT("%s::DummyEvent"), *NodeEvents.Key->GetName()));
-				EventParams.ObservedObject = NodeEvents.Key;
-				EventParams.OwningGraphName = NAME_None;
-				EventParams.DisplayName = NodeEvents.Key->GetNodeTitle(ENodeTitleType::MenuTitle);
-				EventParams.Tooltip = LOCTEXT("NavigateToEventLocationHyperlink_ToolTip", "Navigate to the Event");
-				EventParams.NodeFlags = FunctionContext->IsInheritedContext() ? (EScriptExecutionNodeFlags::Event|EScriptExecutionNodeFlags::RuntimeEvent|EScriptExecutionNodeFlags::InheritedEvent) : 
-																				(EScriptExecutionNodeFlags::Event|EScriptExecutionNodeFlags::RuntimeEvent);
-				GetNodeCustomizations(EventParams);
-				const FSlateBrush* EventIcon = NodeEvents.Key->ShowPaletteIconOnNode() ? NodeEvents.Key->GetIconAndTint(EventParams.IconColor).GetOptionalIcon() :
-																						 FEditorStyle::GetBrush(TEXT("BlueprintProfiler.BPNode"));
-				EventParams.Icon = const_cast<FSlateBrush*>(EventIcon);
-				EventExecNode = CreateExecutionNode(EventParams);
-				// Add entry points.
-				AddEntryPoint(EventExecNode);
-				BlueprintContextIn->AddEventNode(FunctionContext, EventExecNode);
-			}
 			// Create the events for the pins
-			for (FPinDelegateDesc EventDesc : NodeEvents.Value)
+
+			for (FPinDelegateDesc& EventDesc : NodeEvents.Value)
 			{
-				FScriptExecNodeParams PinParams;
-				PinParams.SampleFrequency = 1;
-				PinParams.NodeName = GetUniquePinName(EventDesc.DelegatePin);
-				PinParams.ObservedPin = EventDesc.DelegatePin;
-				PinParams.OwningGraphName = NAME_None;
-				PinParams.DisplayName = EventDesc.DelegatePin->GetDisplayName();
-				PinParams.Tooltip = LOCTEXT("ExecPin_ExpandExecutionPath_ToolTip", "Expand execution path");
-				PinParams.NodeFlags = EScriptExecutionNodeFlags::ExecPin|EScriptExecutionNodeFlags::EventPin;
-				PinParams.IconColor = GetDefault<UGraphEditorSettings>()->DelegatePinTypeColor;
-				const bool bPinLinked = EventDesc.DelegatePin->LinkedTo.Num() > 0;
-				const FSlateBrush* Icon = bPinLinked ?	FEditorStyle::GetBrush(TEXT("BlueprintProfiler.BPPinConnected")) : 
-														FEditorStyle::GetBrush(TEXT("BlueprintProfiler.BPPinDisconnected"));
-				PinParams.Icon = const_cast<FSlateBrush*>(Icon);
-				TSharedPtr<FScriptExecutionNode> PinExecNode = CreateExecutionNode(PinParams);
-				AddEntryPoint(PinExecNode);
-				// Register the function context as a handler for the event.
-				BlueprintContextIn->RegisterEventContext(EventDesc.EventName, FunctionContext);
-				// Register exec node under pin name
-				ExecutionNodes.Add(EventDesc.EventName) = EventExecNode.IsValid() ? EventExecNode : PinExecNode;
+				if (EventDesc.Context.IsValid())
+				{
+					TSharedPtr<FScriptExecutionNode> PinExecNode = EventDesc.Context.ProfilerNode->GetLinkedNodeByScriptOffset(EventDesc.ScriptOffset);
+					if (PinExecNode.IsValid())
+					{
+						// Modify the pin node to mark it as an event/delegate pin.
+						PinExecNode->AddFlags(EScriptExecutionNodeFlags::ExecPin | EScriptExecutionNodeFlags::EventPin);
+						PinExecNode->SetIconColor(GetDefault<UGraphEditorSettings>()->DelegatePinTypeColor);
+						// Register the function context as a handler for the event.
+						ParentBlueprintContext->RegisterEventContext(EventDesc.EventName, AsShared());
+					}
+				}
+				else
+				{
+					if (!EventExecNode.IsValid())
+					{
+						// Check if this node requires an event node creating.
+						bool bCreateEventNode = true;
+						for (auto Pin : NodeEvents.Key->Pins)
+						{
+							if (Pin->Direction == EGPD_Input && Pin->PinType.PinCategory == UEdGraphSchema_K2::PC_Exec)
+							{
+								bCreateEventNode = false;
+								break;
+							}
+						}
+						if (bCreateEventNode)
+						{
+							// Setup the basic exec node params.
+							FScriptExecNodeParams EventParams;
+							EventParams.SampleFrequency = 1;
+							EventParams.NodeName = FName(*FString::Printf(TEXT("%s::DummyEvent"), *NodeEvents.Key->GetName()));
+							EventParams.ObservedObject = NodeEvents.Key;
+							EventParams.OwningGraphName = NAME_None;
+							EventParams.DisplayName = NodeEvents.Key->GetNodeTitle(ENodeTitleType::MenuTitle);
+							EventParams.Tooltip = LOCTEXT("NavigateToEventLocationHyperlink_ToolTip", "Navigate to the Event");
+							EventParams.NodeFlags = EventDesc.Context.ProfilerContext->IsInheritedContext() ? (EScriptExecutionNodeFlags::Event | EScriptExecutionNodeFlags::RuntimeEvent | EScriptExecutionNodeFlags::InheritedEvent) :
+								(EScriptExecutionNodeFlags::Event | EScriptExecutionNodeFlags::RuntimeEvent);
+							GetNodeCustomizations(EventParams);
+							const FSlateBrush* EventIcon = NodeEvents.Key->ShowPaletteIconOnNode() ? NodeEvents.Key->GetIconAndTint(EventParams.IconColor).GetOptionalIcon() :
+								FEditorStyle::GetBrush(TEXT("BlueprintProfiler.BPNode"));
+							EventParams.Icon = const_cast<FSlateBrush*>(EventIcon);
+							EventExecNode = CreateExecutionNode(EventParams);
+							// Add entry points.
+							AddEntryPoint(EventExecNode);
+							ExecutionNodes.Add(EventDesc.EventName) = EventExecNode;
+							ParentBlueprintContext->AddEventNode(EventDesc.Context.ProfilerContext, EventExecNode);
+						}
+					}
+
+					if (EventExecNode.IsValid())
+					{
+						// Create the events for the pins
+						FScriptExecNodeParams PinParams;
+						PinParams.SampleFrequency = 1;
+						PinParams.NodeName = GetUniquePinName(EventDesc.DelegatePin);
+						PinParams.ObservedPin = EventDesc.DelegatePin;
+						PinParams.ObservedObject = NodeEvents.Key;
+						PinParams.OwningGraphName = NAME_None;
+						PinParams.DisplayName = EventDesc.DelegatePin->GetDisplayName();
+						PinParams.Tooltip = LOCTEXT("ExecPin_ExpandExecutionPath_ToolTip", "Expand execution path");
+						PinParams.NodeFlags = EScriptExecutionNodeFlags::ExecPin | EScriptExecutionNodeFlags::EventPin;
+						PinParams.IconColor = GetDefault<UGraphEditorSettings>()->DelegatePinTypeColor;
+						const bool bPinLinked = EventDesc.DelegatePin->LinkedTo.Num() > 0;
+						const FSlateBrush* Icon = bPinLinked ? FEditorStyle::GetBrush(TEXT("BlueprintProfiler.BPPinConnected")) :
+							FEditorStyle::GetBrush(TEXT("BlueprintProfiler.BPPinDisconnected"));
+						PinParams.Icon = const_cast<FSlateBrush*>(Icon);
+
+						TSharedPtr<FScriptExecutionNode> PinExecNode = CreateExecutionNode(PinParams);
+						AddEntryPoint(PinExecNode);
+						EventExecNode->AddChildNode(PinExecNode);
+						if (bPinLinked)
+						{
+							TSharedPtr<FScriptExecutionNode> LinkedNode = MapNodeExecution(EventDesc.DelegatePin->LinkedTo[0]->GetOwningNode());
+							PinExecNode->AddLinkedNode(EventDesc.ScriptOffset, LinkedNode);
+						}
+
+						ParentBlueprintContext->RegisterEventContext(EventDesc.EventName, FunctionContext);
+
+						FNodeExecutionContext& NodeContext = GetNodeExecutionContext(EventDesc.ScriptOffset);
+						NodeContext.GraphNode = EventDesc.DelegatePin->GetOwningNode();
+						NodeContext.ProfilerNode = PinExecNode;
+					}
+				}
 			}
 		}
 	}
@@ -1697,6 +1758,16 @@ bool FBlueprintFunctionContext::GetProfilerContextFromScriptOffset(const int32 S
 	return ExecContextOut.IsValid();
 }
 
+bool FBlueprintFunctionContext::IsUbergraphFunction() const
+{
+	bool bResult = false;
+	if (UBlueprintGeneratedClass* BPGC = Cast<UBlueprintGeneratedClass>(BlueprintClass.Get()))
+	{
+		bResult = Function.Get() == BPGC->UberGraphFunction;
+	}
+	return bResult;
+}
+
 const UEdGraphNode* FBlueprintFunctionContext::GetNodeFromCodeLocation(const int32 ScriptOffset)
 {
 	TWeakObjectPtr<const UEdGraphNode>& Result = ScriptOffsetToNodes.FindOrAdd(ScriptOffset);
@@ -2070,6 +2141,17 @@ void FBlueprintTunnelInstanceContext::MapTunnelContext(TSharedPtr<FBlueprintFunc
 			InputPins.Add(PinToTrace);
 			MapInputPins(PureChainRootNode, InputPins);
 		}
+		// Map any composite events
+		if (TunnelInstance && TunnelInstanceNode->IsA<UK2Node_Composite>())
+		{
+			MapCompositeEvents(TunnelGraph);
+		}
+		// Detect any cyclic links
+		TArray<TSharedPtr<FScriptExecutionNode>> Filter;
+		for (auto EntryPoint : EntryPoints)
+		{
+			DetectCyclicLinks(EntryPoint, Filter);
+		}
 		// Register graph nodes to context and discover valid script offsets
 		TSet<const UEdGraphNode*> ProcessedNodes;
 		for (auto ExecNode : ExecutionNodes)
@@ -2099,6 +2181,51 @@ void FBlueprintTunnelInstanceContext::MapTunnelContext(TSharedPtr<FBlueprintFunc
 	}
 }
 
+void FBlueprintTunnelInstanceContext::MapCompositeEvents(UEdGraph* CompositeGraph)
+{
+	// Map events
+	TArray<UK2Node_Event*> CompositeEventNodes;
+	CompositeGraph->GetNodesOfClass<UK2Node_Event>(CompositeEventNodes);
+	for (auto EventNode : CompositeEventNodes)
+	{
+		if (EventNode->IsNodeEnabled())
+		{
+			const FName EventName = GetScopedEventName(EventNode->GetFunctionName());
+			FScriptExecNodeParams EventParams;
+			EventParams.SampleFrequency = 1;
+			EventParams.NodeName = EventName;
+			EventParams.ObservedObject = EventNode;
+			if (bIsInheritedContext)
+			{
+				EventParams.Tooltip = LOCTEXT("NavigateToInheritedEventLocationHyperlink_ToolTip", "Navigate to the Inherited Event");
+				EventParams.IconColor = GetDefault<UGraphEditorSettings>()->ParentFunctionCallNodeTitleColor;
+				EventParams.NodeFlags = EScriptExecutionNodeFlags::Event|EScriptExecutionNodeFlags::RuntimeEvent|EScriptExecutionNodeFlags::InheritedEvent;
+				EventParams.DisplayName = FText::FromString(FString::Printf(TEXT("%s (%s)"), *EventNode->GetNodeTitle(ENodeTitleType::MenuTitle).ToString(), *BlueprintContext.Pin()->GetBlueprint()->GetName()));
+			}
+			else
+			{
+				EventParams.Tooltip = LOCTEXT("NavigateToEventLocationHyperlink_ToolTip", "Navigate to the Event");
+				EventParams.IconColor = GetDefault<UGraphEditorSettings>()->EventNodeTitleColor;
+				EventParams.NodeFlags = EScriptExecutionNodeFlags::Event|EScriptExecutionNodeFlags::RuntimeEvent;
+				EventParams.DisplayName = EventNode->GetNodeTitle(ENodeTitleType::MenuTitle);
+			}
+			GetNodeCustomizations(EventParams);
+			const FSlateBrush* EventIcon = EventNode->ShowPaletteIconOnNode() ?	EventNode->GetIconAndTint(EventParams.IconColor).GetOptionalIcon() :
+																				FEditorStyle::GetBrush(TEXT("BlueprintProfiler.BPNode"));
+			EventParams.Icon = const_cast<FSlateBrush*>(EventIcon);
+			TSharedPtr<FScriptExecutionNode> EventExecNode = CreateExecutionNode(EventParams);
+			AddEntryPoint(EventExecNode);
+			BlueprintContext.Pin()->AddEventNode(AsShared(), EventExecNode);
+			// Map the event
+			if (EventExecNode.IsValid())
+			{
+				TSharedPtr<FScriptExecutionNode> EventEntryPoint = MapNodeExecution(EventNode);
+				EventExecNode->AddChildNode(EventEntryPoint);
+			}
+		}
+	}
+}
+
 void FBlueprintTunnelInstanceContext::MapTunnelIO(TMap<UEdGraphPin*, UEdGraphPin*>& PurePins)
 {
 	if (UK2Node_Tunnel* TunnelInstanceGraphNode = TunnelInstanceNode.Get())
@@ -2473,9 +2600,16 @@ void FBlueprintTunnelInstanceContext::DiscoverExitSites(TSharedPtr<FScriptExecut
 	}
 	else
 	{
-		for (auto ChildIter : MappedNode->GetChildNodes())
+		// Don't map inside tunnels during discovery!
+		if (!MappedNode->IsTunnelEntry())
 		{
-			DiscoverExitSites(ChildIter);
+			for (auto ChildIter : MappedNode->GetChildNodes())
+			{
+				if (!ChildIter->IsPureChain())
+				{
+					DiscoverExitSites(ChildIter);
+				}
+			}
 		}
 		for (auto LinkIter : MappedNode->GetLinkedNodes())
 		{
@@ -2557,7 +2691,6 @@ const int32 FBlueprintTunnelInstanceContext::GetCodeLocationFromPin(const UEdGra
 			}
 			else if (CodeLocations.Num() > 1)
 			{
-				TArray<UEdGraphNode*> MacrosAtLocation;
 				for (const int32 CodeLocation : CodeLocations)
 				{
 					if (IsCodeLocationValidForThisTunnel(CodeLocation))
@@ -2885,6 +3018,7 @@ bool FScriptEventPlayback::Process(const TArray<FScriptInstrumentedEvent>& Signa
 			check (EventNode.IsValid());
 			TSharedPtr<FScriptPerfData> PerfData = EventNode->GetOrAddPerfDataByInstanceAndTracePath(InstanceName, EventTracePath);
 			PerfData->AddEventTiming(EventTiming.Value);
+			EventNode->AddFlags(EScriptExecutionNodeFlags::RequiresRefresh);
 		}
 		EventTimings.Reset();
 		// Process Node map timings -- this can probably be rolled back into submission during the initial processing and lose this extra iteration.
diff --git a/Engine/Source/Developer/BlueprintProfiler/Public/ScriptInstrumentationPlayback.h b/Engine/Source/Developer/BlueprintProfiler/Public/ScriptInstrumentationPlayback.h
--- a/Engine/Source/Developer/BlueprintProfiler/Public/ScriptInstrumentationPlayback.h
+++ b/Engine/Source/Developer/BlueprintProfiler/Public/ScriptInstrumentationPlayback.h
@@ -201,6 +201,9 @@ public:
 	/** Returns if the function context is fully formed */
 	bool IsContextValid() const { return Function.IsValid() && BlueprintClass.IsValid(); }
 
+	/** Returns if the function context is represents the ubergraph */
+	bool IsUbergraphFunction() const;
+
 	/** Returns if the function context is inherited */
 	bool IsInheritedContext() const { return bIsInheritedContext; }
 
@@ -324,10 +327,10 @@ protected:
 	void AddExitPoint(TSharedPtr<FScriptExecutionNode> ExitPoint);
 
 	/** Creates an event for delegate pin entry points */
-	void CreateDelegatePinEvents(TSharedPtr<FBlueprintExecutionContext> BlueprintContextIn, const TMap<FName, FEdGraphPinReference>& PinEvents);
+	void CreateDelegatePinEvents();
 
 	/** Processes and detects any cyclic links making the linkage safe for traversal */
-	bool DetectCyclicLinks(TSharedPtr<FScriptExecutionNode> ExecNode, TSet<TSharedPtr<FScriptExecutionNode>>& Filter);
+	bool DetectCyclicLinks(TSharedPtr<FScriptExecutionNode> ExecNode, TArray<TSharedPtr<FScriptExecutionNode>>& Filter);
 
 	/** Add child function context */
 	void AddChildFunctionContext(const FName FunctionNameIn, TSharedPtr<FBlueprintFunctionContext> ChildContext);
@@ -444,6 +447,9 @@ private:
 	virtual FName GetTunnelInstanceFunctionName(const UEdGraphNode* GraphNode) const override;
 	// ~FBlueprintFunctionContext End
 	
+	/** Maps any events nested inside a composite node in the event graph */
+	void MapCompositeEvents(UEdGraph* CompositeGraph);
+
 	/** Maps input and output pins in a tunnel graph */
 	void MapTunnelIO(TMap<UEdGraphPin*, UEdGraphPin*>& PurePinsOut);
 
diff --git a/Engine/Source/Editor/BlueprintGraph/Classes/EdGraphSchema_K2.h b/Engine/Source/Editor/BlueprintGraph/Classes/EdGraphSchema_K2.h
--- a/Engine/Source/Editor/BlueprintGraph/Classes/EdGraphSchema_K2.h
+++ b/Engine/Source/Editor/BlueprintGraph/Classes/EdGraphSchema_K2.h
@@ -175,8 +175,14 @@ public:
 	static const FName MD_ArrayParam;
 	static const FName MD_ArrayDependentParam;
 
+	/** Metadata that flags TSet parameters that will have their type determined at blueprint compile time */
 	static const FName MD_SetParam;
 
+	/** Metadata that flags TMap function parameters that will have their type determined at blueprint compile time */
+	static const FName MD_MapParam;
+	static const FName MD_MapKeyParam;
+	static const FName MD_MapValueParam;
+
 	/** Metadata that identifies an integral property as a bitmask. */
 	static const FName MD_Bitmask;
 	/** Metadata that associates a bitmask property with a bitflag enum. */
diff --git a/Engine/Source/Editor/BlueprintGraph/Classes/K2Node_AddComponent.h b/Engine/Source/Editor/BlueprintGraph/Classes/K2Node_AddComponent.h
--- a/Engine/Source/Editor/BlueprintGraph/Classes/K2Node_AddComponent.h
+++ b/Engine/Source/Editor/BlueprintGraph/Classes/K2Node_AddComponent.h
@@ -64,16 +64,15 @@ class UK2Node_AddComponent : public UK2Node_CallFunction
 		return FindPinChecked(TEXT("bManualAttachment"));
 	}
 
+	/** Helper method used to generate a new, unique component template name. */
+	FName MakeNewComponentTemplateName(UObject* InOuter, UClass* InComponentClass);
+
 	/** Helper method used to instantiate a new component template after duplication. */
 	BLUEPRINTGRAPH_API void MakeNewComponentTemplate();
 
 	/** Static name of function to call */
 	static BLUEPRINTGRAPH_API FName GetAddComponentFunctionName();
 
-protected:
-	/** Helper method used to generate a new, unique component template name. */
-	static FName MakeNewComponentTemplateName(UObject* InOuter, UClass* InComponentClass);
-
 private: 
 	UEdGraphPin* GetTemplateNamePin() const
 	{
diff --git a/Engine/Source/Editor/BlueprintGraph/Classes/K2Node_IfThenElse.h b/Engine/Source/Editor/BlueprintGraph/Classes/K2Node_IfThenElse.h
--- a/Engine/Source/Editor/BlueprintGraph/Classes/K2Node_IfThenElse.h
+++ b/Engine/Source/Editor/BlueprintGraph/Classes/K2Node_IfThenElse.h
@@ -5,7 +5,7 @@
 #include "K2Node.h"
 #include "K2Node_IfThenElse.generated.h"
 
-UCLASS(MinimalAPI, meta=(Keywords = "if bool"))
+UCLASS(MinimalAPI, meta=(Keywords = "if bool branch"))
 class UK2Node_IfThenElse : public UK2Node
 {
 	GENERATED_UCLASS_BODY()
diff --git a/Engine/Source/Editor/BlueprintGraph/Classes/K2Node_Select.h b/Engine/Source/Editor/BlueprintGraph/Classes/K2Node_Select.h
--- a/Engine/Source/Editor/BlueprintGraph/Classes/K2Node_Select.h
+++ b/Engine/Source/Editor/BlueprintGraph/Classes/K2Node_Select.h
@@ -6,7 +6,7 @@
 #include "NodeDependingOnEnumInterface.h"
 #include "K2Node_Select.generated.h"
 
-UCLASS(MinimalAPI, meta=(Keywords = "Ternary If"))
+UCLASS(MinimalAPI, meta=(Keywords = "Ternary Select"))
 class UK2Node_Select : public UK2Node, public INodeDependingOnEnumInterface
 {
 	GENERATED_UCLASS_BODY()
diff --git a/Engine/Source/Editor/BlueprintGraph/Classes/K2Node_Switch.h b/Engine/Source/Editor/BlueprintGraph/Classes/K2Node_Switch.h
--- a/Engine/Source/Editor/BlueprintGraph/Classes/K2Node_Switch.h
+++ b/Engine/Source/Editor/BlueprintGraph/Classes/K2Node_Switch.h
@@ -47,6 +47,8 @@ class UK2Node_Switch : public UK2Node
 	/** Gets the pin type from the schema for the subclass */
 	virtual FEdGraphPinType GetPinType() const { check(false); return FEdGraphPinType(); }
 
+	BLUEPRINTGRAPH_API virtual FEdGraphPinType GetInnerCaseType() const;
+
 	/**
 	 * Adds a new execution pin to a switch node
 	 */
@@ -60,9 +62,9 @@ class UK2Node_Switch : public UK2Node
 	BLUEPRINTGRAPH_API virtual void RemovePinFromSwitchNode(UEdGraphPin* TargetPin);
 
 	/** Getting pin access */
-	BLUEPRINTGRAPH_API UEdGraphPin* GetSelectionPin();
-	BLUEPRINTGRAPH_API UEdGraphPin* GetDefaultPin();
-	BLUEPRINTGRAPH_API UEdGraphPin* GetFunctionPin();
+	BLUEPRINTGRAPH_API UEdGraphPin* GetSelectionPin() const;
+	BLUEPRINTGRAPH_API UEdGraphPin* GetDefaultPin() const;
+	BLUEPRINTGRAPH_API UEdGraphPin* GetFunctionPin() const;
 
 	BLUEPRINTGRAPH_API static FString	GetSelectionPinName();
 
diff --git a/Engine/Source/Editor/BlueprintGraph/Private/BlueprintActionFilter.cpp b/Engine/Source/Editor/BlueprintGraph/Private/BlueprintActionFilter.cpp
--- a/Engine/Source/Editor/BlueprintGraph/Private/BlueprintActionFilter.cpp
+++ b/Engine/Source/Editor/BlueprintGraph/Private/BlueprintActionFilter.cpp
@@ -1429,9 +1429,10 @@ static bool BlueprintActionFilterImpl::IsMissingMatchingPinParam(FBlueprintActio
 {
 	bool bIsFilteredOut = false;
 
-	// we have a separate pin tests for function/property nodes (IsFunctionMissingPinParam/IsMissmatchedPropertyType)
-	bool const bTestPinCompatibility = (BlueprintAction.GetAssociatedProperty() == nullptr) &&
-		(BlueprintAction.GetAssociatedFunction() == nullptr);
+	// we have a separate pin tests for function/property nodes (IsFunctionMissingPinParam/IsMissmatchedPropertyType). Note that we only skip
+	// this test for functions with bindings (because it does not handle getting templates for all binding nodes). By running this for 
+	// other functions we ensure that IsConnectionDisallowed is honored.
+	bool const bTestPinCompatibility = (BlueprintAction.GetAssociatedProperty() == nullptr) && BlueprintAction.GetBindings().Num() == 0;
 
 	if (bTestPinCompatibility)
 	{
diff --git a/Engine/Source/Editor/BlueprintGraph/Private/BlueprintComponentNodeSpawner.cpp b/Engine/Source/Editor/BlueprintGraph/Private/BlueprintComponentNodeSpawner.cpp
--- a/Engine/Source/Editor/BlueprintGraph/Private/BlueprintComponentNodeSpawner.cpp
+++ b/Engine/Source/Editor/BlueprintGraph/Private/BlueprintComponentNodeSpawner.cpp
@@ -188,7 +188,7 @@ UEdGraphNode* UBlueprintComponentNodeSpawner::Invoke(UEdGraph* ParentGraph, FBin
 
 		UBlueprint* Blueprint = NewNode->GetBlueprint();
 
-		FName DesiredComponentName = MakeUniqueObjectName(Blueprint->GeneratedClass, Class, *Class->GetDisplayNameText().ToString());
+		FName DesiredComponentName = NewNode->MakeNewComponentTemplateName(Blueprint->GeneratedClass, Class);
 		UActorComponent* ComponentTemplate = NewObject<UActorComponent>(Blueprint->GeneratedClass, Class, DesiredComponentName, RF_ArchetypeObject | RF_Public | RF_Transactional);
 
 		Blueprint->ComponentTemplates.Add(ComponentTemplate);
diff --git a/Engine/Source/Editor/BlueprintGraph/Private/CallFunctionHandler.cpp b/Engine/Source/Editor/BlueprintGraph/Private/CallFunctionHandler.cpp
--- a/Engine/Source/Editor/BlueprintGraph/Private/CallFunctionHandler.cpp
+++ b/Engine/Source/Editor/BlueprintGraph/Private/CallFunctionHandler.cpp
@@ -90,23 +90,6 @@ void FKCHandler_CallFunction::CreateFunctionCallStatement(FKismetFunctionContext
 			CompilerContext.MessageLog.Error(*LOCTEXT("ContainsLatentCall_Error", "@@ contains a latent call, which cannot exist outside of the event graph").ToString(), Node);
 		}
 
-		// Check access specifier
-		const uint32 AccessSpecifier = Function->FunctionFlags & FUNC_AccessSpecifiers;
-		if(FUNC_Private == AccessSpecifier)
-		{
-			if(Function->GetOuter() != Context.NewClass)
-			{
-				CompilerContext.MessageLog.Warning(*LOCTEXT("PrivateFunctionCall_Error", "Function @@ is private and cannot be called outside its class").ToString(), Node);
-			}
-		}
-		else if(FUNC_Protected == AccessSpecifier)
-		{
-			if( !Context.NewClass->IsChildOf( Cast<UStruct>( Function->GetOuter() ) ) )
-			{
-				CompilerContext.MessageLog.Warning(*LOCTEXT("ProtectedFunctionCall_Error", "Function @@ is protected and can be called only from its class or subclasses").ToString(), Node);
-			}
-		}
-
 		UEdGraphPin* LatentInfoPin = NULL;
 
 		TMap<FName, FString>* MetaData = UMetaData::GetMapForObject(Function);
diff --git a/Engine/Source/Editor/BlueprintGraph/Private/EdGraphSchema_K2.cpp b/Engine/Source/Editor/BlueprintGraph/Private/EdGraphSchema_K2.cpp
--- a/Engine/Source/Editor/BlueprintGraph/Private/EdGraphSchema_K2.cpp
+++ b/Engine/Source/Editor/BlueprintGraph/Private/EdGraphSchema_K2.cpp
@@ -4,6 +4,7 @@
 
 #include "UObject/DevObjectVersion.h"
 #include "Engine/LevelScriptBlueprint.h"
+#include "Kismet/BlueprintSetLibrary.h"
 #include "Kismet/KismetSystemLibrary.h"
 #include "Kismet/KismetArrayLibrary.h"
 #include "Kismet/KismetMathLibrary.h"
@@ -35,6 +36,7 @@
 #include "K2Node_SetFieldsInStruct.h"
 #include "K2Node_ConvertAsset.h"
 #include "GenericCommands.h"
+#include "BlueprintSupport.h" // for FLegacyEditorOnlyBlueprintOptions::IsTypeProhibited()
 
 
 //////////////////////////////////////////////////////////////////////////
@@ -107,6 +109,11 @@ const FName FBlueprintMetadata::MD_ArrayDependentParam(TEXT("ArrayTypeDependentP
 
 const FName FBlueprintMetadata::MD_SetParam(TEXT("SetParam"));
 
+// Each of these is a | separated list of param names:
+const FName FBlueprintMetadata::MD_MapParam(TEXT("MapParam"));
+const FName FBlueprintMetadata::MD_MapKeyParam(TEXT("MapKeyParam"));
+const FName FBlueprintMetadata::MD_MapValueParam(TEXT("MapValueParam"));
+
 const FName FBlueprintMetadata::MD_Bitmask(TEXT("Bitmask"));
 const FName FBlueprintMetadata::MD_BitmaskEnum(TEXT("BitmaskEnum"));
 const FName FBlueprintMetadata::MD_Bitflags(TEXT("Bitflags"));
@@ -1113,8 +1120,7 @@ bool UEdGraphSchema_K2::IsAllowableBlueprintVariableType(const class UClass* InC
 			return true;
 		}
 
-		static const FBoolConfigValueHelper NotBlueprintType(TEXT("EditoronlyBP"), TEXT("bBlueprintIsNotBlueprintType"));
-		if (NotBlueprintType && InClass->IsChildOf(UBlueprint::StaticClass()))
+		if (FLegacyEditorOnlyBlueprintOptions::IsTypeProhibited(InClass))
 		{
 			return false;
 		}
@@ -2402,6 +2408,13 @@ bool UEdGraphSchema_K2::SearchForAutocastFunction(const UEdGraphPin* OutputPin,
 		OutputPin->PinType.bIsMap != InputPin->PinType.bIsMap ||
 		OutputPin->PinType.bIsSet != InputPin->PinType.bIsSet)
 	{
+		if (OutputPin->PinType.bIsSet && InputPin->PinType.bIsArray)
+		{
+			UFunction* Function = UBlueprintSetLibrary::StaticClass()->FindFunctionByName(GET_MEMBER_NAME_CHECKED(UBlueprintSetLibrary, Set_ToArray));
+			TargetFunction = Function->GetFName();
+			FunctionOwner = Function->GetOwnerClass();
+			return true;
+		}
 		return false;
 	}
 
@@ -2474,7 +2487,7 @@ bool UEdGraphSchema_K2::FindSpecializedConversionNode(const UEdGraphPin* OutputP
 	TargetNode = NULL;
 
 	// Conversion for scalar -> array
-	if( (!OutputPin->PinType.bIsArray && InputPin->PinType.bIsArray) && ArePinTypesCompatible(OutputPin->PinType, InputPin->PinType, NULL, true) )
+	if( (!OutputPin->PinType.IsContainer() && InputPin->PinType.bIsArray) && ArePinTypesCompatible(OutputPin->PinType, InputPin->PinType, NULL, true))
 	{
 		bCanConvert = true;
 		if(bCreateNode)
@@ -3313,6 +3326,10 @@ bool UEdGraphSchema_K2::ConvertPropertyToPinType(const UProperty* Property, /*ou
 		|| UK2Node_CallFunction::IsWildcardProperty(Function, Property))
 	{
 		TypeOut.PinCategory = PC_Wildcard;
+		if(MapProperty)
+		{
+			TypeOut.PinValueType.TerminalCategory = PC_Wildcard;
+		}
 	}
 	else if (const UMulticastDelegateProperty* MulticastDelegateProperty = Cast<const UMulticastDelegateProperty>(TestProperty))
 	{
@@ -3412,6 +3429,25 @@ FText UEdGraphSchema_K2::TypeToText(UProperty* const Property)
 			return FText::Format(LOCTEXT("ArrayPropertyText", "Array of {ArrayType}"), Args); 
 		}
 	}
+	else if (USetProperty* Set = Cast<USetProperty>(Property))
+	{
+		if (Set->ElementProp)
+		{
+			FFormatNamedArguments Args;
+			Args.Add(TEXT("SetType"), TypeToText(Set->ElementProp));
+			return FText::Format(LOCTEXT("ArrayPropertyText", "Set of {SetType}"), Args);
+		}
+	}
+	else if (UMapProperty* Map = Cast<UMapProperty>(Property))
+	{
+		if (Map->KeyProp && Map->ValueProp)
+		{
+			FFormatNamedArguments Args;
+			Args.Add(TEXT("MapKeyType"), TypeToText(Map->KeyProp));
+			Args.Add(TEXT("MapValueType"), TypeToText(Map->ValueProp));
+			return FText::Format(LOCTEXT("ArrayPropertyText", "Map of {MapKeyType} to {MapValueType}"), Args);
+		}
+	}
 	
 	return FText::FromString(Property->GetClass()->GetName());
 }
@@ -3539,7 +3575,7 @@ FText UEdGraphSchema_K2::TypeToText(const FEdGraphPinType& Type)
 		Args.Add(TEXT("KeyTitle"), PropertyText);
 		FText ValueText = TerminalTypeToText(Type.PinValueType.TerminalCategory, Type.PinValueType.TerminalSubCategory, Type.PinValueType.TerminalSubCategoryObject.Get(), Type.PinValueType.bTerminalIsWeakPointer);
 		Args.Add(TEXT("ValueTitle"), ValueText);
-		PropertyText = FText::Format(LOCTEXT("MapAsText", "Map of {PropertyTitle}s to {ValueTitle}s"), Args);
+		PropertyText = FText::Format(LOCTEXT("MapAsText", "Map of {KeyTitle}s to {ValueTitle}s"), Args);
 	}
 	else if ( Type.bIsSet )
 	{
@@ -4170,37 +4206,8 @@ bool UEdGraphSchema_K2::ArePinTypesCompatible(const FEdGraphPinType& Output, con
 	}
 
 	// Pins representing BLueprint objects and subclass of UObject can match when EditoronlyBP.bAllowClassAndBlueprintPinMatching=true (BaseEngine.ini)
-	// It's required for converting all UBlueprint references into UClass.
-	struct FObjClassAndBlueprintHelper
-	{
-	private:
-		bool bAllow;
-	public:
-		FObjClassAndBlueprintHelper() : bAllow(false)
-		{
-			GConfig->GetBool(TEXT("EditoronlyBP"), TEXT("bAllowClassAndBlueprintPinMatching"), bAllow, GEditorIni);
-		}
-
-		bool Match(const FEdGraphPinType& A, const FEdGraphPinType& B, const UEdGraphSchema_K2& Schema) const
-		{
-			if (bAllow && (B.PinCategory == Schema.PC_Object) && (A.PinCategory == Schema.PC_Class))
-			{
-				const bool bAIsObjectClass = (UObject::StaticClass() == A.PinSubCategoryObject.Get());
-				const UClass* BClass = Cast<UClass>(B.PinSubCategoryObject.Get());
-				const bool bBIsBlueprintObj = BClass && BClass->IsChildOf(UBlueprint::StaticClass());
-				return bAIsObjectClass && bBIsBlueprintObj;
-			}
-			return false;
-		}
-	};
-
-	static FObjClassAndBlueprintHelper MatchHelper;
-	if (MatchHelper.Match(Input, Output, *this) || MatchHelper.Match(Output, Input, *this))
-	{
-		return true;
-	}
-
-	return false;
+	// It's required for converting all UBlueprint references into UClass.	
+	return FLegacyEditorOnlyBlueprintUtils::DoPinsMatch(Input, Output);
 }
 
 void UEdGraphSchema_K2::BreakNodeLinks(UEdGraphNode& TargetNode) const
diff --git a/Engine/Source/Editor/BlueprintGraph/Private/K2Node_AddComponent.cpp b/Engine/Source/Editor/BlueprintGraph/Private/K2Node_AddComponent.cpp
--- a/Engine/Source/Editor/BlueprintGraph/Private/K2Node_AddComponent.cpp
+++ b/Engine/Source/Editor/BlueprintGraph/Private/K2Node_AddComponent.cpp
@@ -440,7 +440,7 @@ FName UK2Node_AddComponent::MakeNewComponentTemplateName(UObject* InOuter, UClas
 {
 	FName TemplateName = NAME_None;
 
-	int32 Counter = 1;
+	int32& Counter = GetBlueprint()->ComponentTemplateNameIndex.FindOrAdd(InComponentClass->GetFName());
 	FString ComponentClassNameAsString = InComponentClass->GetName();
 	if (InComponentClass->HasAnyClassFlags(CLASS_CompiledFromBlueprint))
 	{
diff --git a/Engine/Source/Editor/BlueprintGraph/Private/K2Node_CallFunction.cpp b/Engine/Source/Editor/BlueprintGraph/Private/K2Node_CallFunction.cpp
--- a/Engine/Source/Editor/BlueprintGraph/Private/K2Node_CallFunction.cpp
+++ b/Engine/Source/Editor/BlueprintGraph/Private/K2Node_CallFunction.cpp
@@ -2313,6 +2313,85 @@ void UK2Node_CallFunction::InvalidatePinTooltips()
 
 void UK2Node_CallFunction::ConformContainerPins()
 {
+	// helper functions for type propagation:
+	const auto TryReadTypeToPropagate = [](UEdGraphPin* Pin, bool& bOutPropagated, FEdGraphTerminalType& TypeToPropagete)
+	{
+		if (Pin && !bOutPropagated)
+		{
+			if (Pin->LinkedTo.Num() != 0 || Pin->DefaultValue != Pin->AutogeneratedDefaultValue)
+			{
+				bOutPropagated = true;
+				if (Pin->LinkedTo.Num() != 0)
+				{
+					TypeToPropagete = Pin->LinkedTo[0]->GetPrimaryTerminalType();
+				}
+				else
+				{
+					TypeToPropagete = Pin->GetPrimaryTerminalType();
+				}
+			}
+		}
+	};
+
+	const auto TryReadValueTypeToPropagate = [](UEdGraphPin* Pin, bool& bOutPropagated, FEdGraphTerminalType& TypeToPropagete)
+	{
+		if (Pin && !bOutPropagated)
+		{
+			if (Pin->LinkedTo.Num() != 0 || Pin->DefaultValue != Pin->AutogeneratedDefaultValue)
+			{
+				bOutPropagated = true;
+				if (Pin->LinkedTo.Num() != 0)
+				{
+					TypeToPropagete = Pin->LinkedTo[0]->PinType.PinValueType;
+				}
+				else
+				{
+					TypeToPropagete = Pin->PinType.PinValueType;
+				}
+			}
+		}
+	};
+
+	const auto TryPropagateType = [](UEdGraphPin* Pin, const FEdGraphTerminalType& TerminalType, bool bTypeIsAvailable)
+	{
+		if(Pin)
+		{
+			if(bTypeIsAvailable)
+			{
+				Pin->PinType.PinCategory = TerminalType.TerminalCategory;
+				Pin->PinType.PinSubCategory = TerminalType.TerminalSubCategory;
+				Pin->PinType.PinSubCategoryObject = TerminalType.TerminalSubCategoryObject;
+			}
+			else
+			{
+				// reset to wildcard:
+				Pin->PinType.PinCategory = UEdGraphSchema_K2::PC_Wildcard;
+				Pin->PinType.PinSubCategory.Empty();
+				Pin->PinType.PinSubCategoryObject = nullptr;
+				Pin->DefaultValue = TEXT("");
+			}
+		}
+	};
+
+	const auto TryPropagateValueType = [](UEdGraphPin* Pin, const FEdGraphTerminalType& TerminalType, bool bTypeIsAvailable)
+	{
+		if (Pin)
+		{
+			if (bTypeIsAvailable)
+			{
+				Pin->PinType.PinValueType.TerminalCategory = TerminalType.TerminalCategory;
+				Pin->PinType.PinValueType.TerminalSubCategory = TerminalType.TerminalSubCategory;
+				Pin->PinType.PinValueType.TerminalSubCategoryObject = TerminalType.TerminalSubCategoryObject;
+			}
+			else
+			{
+				Pin->PinType.PinValueType.TerminalCategory = UEdGraphSchema_K2::PC_Wildcard;
+				Pin->PinType.PinValueType.TerminalSubCategory.Empty();
+				Pin->PinType.PinValueType.TerminalSubCategoryObject = nullptr;
+			}
+		}
+	};
+		
 	const UFunction* TargetFunction = GetTargetFunction();
 	if (TargetFunction == nullptr)
 	{
@@ -2320,17 +2399,12 @@ void UK2Node_CallFunction::ConformContainerPins()
 	}
 
 	// find any pins marked as SetParam
-	const FString& DependentPinMetaData = TargetFunction->GetMetaData(FBlueprintMetadata::MD_SetParam);
-	if (DependentPinMetaData.IsEmpty())
-	{
-		// early out, no SetParam metadata
-		return;
-	}
+	const FString& SetPinMetaData = TargetFunction->GetMetaData(FBlueprintMetadata::MD_SetParam);
 
 	// useless copies/allocates in this code, could be an optimization target...
 	TArray<FString> SetParamPinGroups;
 	{
-		DependentPinMetaData.ParseIntoArray(SetParamPinGroups, TEXT(","), true);
+		SetPinMetaData.ParseIntoArray(SetParamPinGroups, TEXT(","), true);
 	}
 
 	for (FString& Entry : SetParamPinGroups)
@@ -2350,46 +2424,50 @@ void UK2Node_CallFunction::ConformContainerPins()
 
 		// if nothing is connected (or non-default), reset to wildcard
 		// else, find the first type and propagate to everyone else::
-		bool bAllPinsAreDefaultOrEmpty = true;
+		bool bReadyToPropagatSetType = false;
 		FEdGraphTerminalType TypeToPropagate;
 		for (UEdGraphPin* Pin : ResolvedPins)
 		{
-			if (Pin->LinkedTo.Num() != 0 || Pin->DefaultValue != Pin->AutogeneratedDefaultValue)
+			TryReadTypeToPropagate(Pin, bReadyToPropagatSetType, TypeToPropagate);
+			if(bReadyToPropagatSetType)
 			{
-				bAllPinsAreDefaultOrEmpty = false;
-				if (Pin->LinkedTo.Num() != 0)
-				{
-					TypeToPropagate = Pin->LinkedTo[0]->GetPrimaryTerminalType();
-				}
-				else
-				{
-					TypeToPropagate = Pin->GetPrimaryTerminalType();
-				}
 				break;
 			}
 		}
 
-		if (!bAllPinsAreDefaultOrEmpty)
-		{
-			for (UEdGraphPin* Pin : ResolvedPins)
-			{
-				Pin->PinType.PinCategory = TypeToPropagate.TerminalCategory;
-				Pin->PinType.PinSubCategory = TypeToPropagate.TerminalSubCategory;
-				Pin->PinType.PinSubCategoryObject = TypeToPropagate.TerminalSubCategoryObject;
-			}
-		}
-		else
+		for (UEdGraphPin* Pin : ResolvedPins)
 		{
-			// reset everyone to wildcard:
-			for (UEdGraphPin* Pin : ResolvedPins)
-			{
-				Pin->PinType.PinCategory = UEdGraphSchema_K2::PC_Wildcard;
-				Pin->PinType.PinSubCategory.Empty();
-				Pin->PinType.PinSubCategoryObject = nullptr;
-				Pin->DefaultValue = TEXT("");
-			}
+			TryPropagateType( Pin, TypeToPropagate, bReadyToPropagatSetType );
 		}
 	}
+
+	const FString& MapPinMetaData = TargetFunction->GetMetaData(FBlueprintMetadata::MD_MapParam);
+	const FString& MapKeyPinMetaData = TargetFunction->GetMetaData(FBlueprintMetadata::MD_MapKeyParam);
+	const FString& MapValuePinMetaData = TargetFunction->GetMetaData(FBlueprintMetadata::MD_MapValueParam);
+
+	if(!MapPinMetaData.IsEmpty() || !MapKeyPinMetaData.IsEmpty() || !MapValuePinMetaData.IsEmpty() )
+	{
+		// if the map pin has a connection infer from that, otherwise use the information on the key param and value param:
+		bool bReadyToPropagateKeyType = false;
+		FEdGraphTerminalType KeyTypeToPropagate;
+		bool bReadyToPropagateValueType = false;
+		FEdGraphTerminalType ValueTypeToPropagate;
+
+		UEdGraphPin* MapPin = MapPinMetaData.IsEmpty() ? nullptr : FindPin(MapPinMetaData);
+		UEdGraphPin* MapKeyPin = MapKeyPinMetaData.IsEmpty() ? nullptr : FindPin(MapKeyPinMetaData);
+		UEdGraphPin* MapValuePin = MapValuePinMetaData.IsEmpty() ? nullptr : FindPin(MapValuePinMetaData);
+
+		TryReadTypeToPropagate(MapPin, bReadyToPropagateKeyType, KeyTypeToPropagate);
+		TryReadValueTypeToPropagate(MapPin, bReadyToPropagateValueType, ValueTypeToPropagate);
+		TryReadTypeToPropagate(MapKeyPin, bReadyToPropagateKeyType, KeyTypeToPropagate);
+		TryReadTypeToPropagate(MapValuePin, bReadyToPropagateValueType, ValueTypeToPropagate);
+
+		TryPropagateType(MapPin, KeyTypeToPropagate, bReadyToPropagateKeyType);
+		TryPropagateType(MapKeyPin, KeyTypeToPropagate, bReadyToPropagateKeyType);
+
+		TryPropagateValueType(MapPin, ValueTypeToPropagate, bReadyToPropagateValueType);
+		TryPropagateType(MapValuePin, ValueTypeToPropagate, bReadyToPropagateValueType);
+	}
 }
 
 FText UK2Node_CallFunction::GetToolTipHeading() const
@@ -2566,7 +2644,7 @@ bool UK2Node_CallFunction::IsWildcardProperty(const UFunction* InFunction, const
 {
 	if (InProperty)
 	{
-		return FEdGraphUtilities::IsSetParam(InFunction, InProperty->GetName());
+		return FEdGraphUtilities::IsSetParam(InFunction, InProperty->GetName()) || FEdGraphUtilities::IsMapParam(InFunction, InProperty->GetName());
 	}
 	return false;
 }
@@ -2591,6 +2669,27 @@ bool UK2Node_CallFunction::IsConnectionDisallowed(const UEdGraphPin* MyPin, cons
 			bIsDisallowed = true;
 			OutReason = LOCTEXT("PinConnectionDisallowed", "This parameter is for internal use only.").ToString();
 		}
+		else if (UFunction* TargetFunction = GetTargetFunction())
+		{
+			if( 
+				// Strictly speaking this first check is not needed, but by not disabling the connection here we get a better reason later:
+				(	(	FEdGraphUtilities::IsSetParam(TargetFunction, MyPin->PinName) &&
+					(OtherPin->PinType.IsContainer() && !MyPin->PinType.bIsSet) ) ||
+					(	FEdGraphUtilities::IsMapParam(TargetFunction, MyPin->PinName) &&
+					(OtherPin->PinType.IsContainer() && !MyPin->PinType.bIsMap) )  )
+				&& 
+				// make sure we don't allow connections of mismatched container types (e.g. maps to arrays)
+				( 
+					OtherPin->PinType.bIsMap != MyPin->PinType.bIsMap || 
+					OtherPin->PinType.bIsSet != MyPin->PinType.bIsSet ||
+					OtherPin->PinType.bIsArray != MyPin->PinType.bIsArray 
+				)
+			)
+			{
+				bIsDisallowed = true;
+				OutReason = LOCTEXT("PinSetConnectionDisallowed", "Containers of containers are not supported - consider wrapping a container in a Structure object").ToString();
+			}
+		}
 	}
 
 	return bIsDisallowed;
diff --git a/Engine/Source/Editor/BlueprintGraph/Private/K2Node_Literal.cpp b/Engine/Source/Editor/BlueprintGraph/Private/K2Node_Literal.cpp
--- a/Engine/Source/Editor/BlueprintGraph/Private/K2Node_Literal.cpp
+++ b/Engine/Source/Editor/BlueprintGraph/Private/K2Node_Literal.cpp
@@ -188,7 +188,7 @@ void UK2Node_Literal::GetMenuActions(FBlueprintActionDatabaseRegistrar& ActionRe
 			UiSpecOut->MenuName = FText::Format( NSLOCTEXT("K2Node", "LiteralTitle", "Create a Reference to {0}"), 
 				FText::FromString(ActorObj->GetActorLabel()) );
 
-			FSlateIcon Icon = FClassIconFinder::FindSlateIconForActor(ActorObj);
+			const FSlateIcon Icon = FSlateIconFinder::FindIconForClass(ActorObj->GetClass());
 			if (Icon.IsSet())
 			{
 				UiSpecOut->Icon = Icon;
@@ -283,15 +283,7 @@ FSlateIcon UK2Node_Literal::GetIconAndTint(FLinearColor& OutColor) const
 {
 	if(ObjectRef != NULL)
 	{
-		AActor* Actor = Cast<AActor>(ObjectRef);
-		if(Actor != NULL)
-		{
-			return FClassIconFinder::FindSlateIconForActor(Actor);
-		}
-		else
-		{
-			return FSlateIconFinder::FindIconForClass(ObjectRef->GetClass());
-		}	
+		return FSlateIconFinder::FindIconForClass(ObjectRef->GetClass());	
 	}
 	return Super::GetIconAndTint(OutColor);
 }
diff --git a/Engine/Source/Editor/BlueprintGraph/Private/K2Node_Switch.cpp b/Engine/Source/Editor/BlueprintGraph/Private/K2Node_Switch.cpp
--- a/Engine/Source/Editor/BlueprintGraph/Private/K2Node_Switch.cpp
+++ b/Engine/Source/Editor/BlueprintGraph/Private/K2Node_Switch.cpp
@@ -97,7 +97,7 @@ public:
 					// Create a term for the switch case value
 					FBPTerminal* CaseValueTerm = new (Context.Literals) FBPTerminal();
 					CaseValueTerm->Name = Pin->PinName;
-					CaseValueTerm->Type = SelectionPin->PinType;
+					CaseValueTerm->Type = SwitchNode->GetInnerCaseType();
 					CaseValueTerm->SourcePin = Pin;
 					CaseValueTerm->bIsLiteral = true;
 
@@ -293,19 +293,19 @@ void UK2Node_Switch::CreateFunctionPin()
 	}
 }
 
-UEdGraphPin* UK2Node_Switch::GetFunctionPin()
+UEdGraphPin* UK2Node_Switch::GetFunctionPin() const
 {
 	//@TODO: Should probably use a specific index, though FindPin starts at 0, so this won't *currently* conflict with user created pins
 	return FindPin(FunctionName.ToString());
 }
 
-UEdGraphPin* UK2Node_Switch::GetSelectionPin()
+UEdGraphPin* UK2Node_Switch::GetSelectionPin() const
 {
 	//@TODO: Should probably use a specific index, though FindPin starts at 0, so this won't *currently* conflict with user created pins
 	return FindPin(SelectionPinName);
 }
 
-UEdGraphPin* UK2Node_Switch::GetDefaultPin()
+UEdGraphPin* UK2Node_Switch::GetDefaultPin() const
 {
 	return (bHasDefaultPin)
 		? Pins[0]
@@ -328,4 +328,13 @@ FText UK2Node_Switch::GetMenuCategory() const
 	return CachedCategory;
 }
 
+FEdGraphPinType UK2Node_Switch::GetInnerCaseType() const
+{
+	UEdGraphPin* SelectionPin = GetSelectionPin();
+	if (ensure(SelectionPin))
+	{
+		return SelectionPin->PinType;
+	}
+	return FEdGraphPinType();
+}
 #undef LOCTEXT_NAMESPACE
diff --git a/Engine/Source/Editor/BlueprintGraph/Private/K2Node_TunnelBoundary.cpp b/Engine/Source/Editor/BlueprintGraph/Private/K2Node_TunnelBoundary.cpp
--- a/Engine/Source/Editor/BlueprintGraph/Private/K2Node_TunnelBoundary.cpp
+++ b/Engine/Source/Editor/BlueprintGraph/Private/K2Node_TunnelBoundary.cpp
@@ -207,21 +207,27 @@ void UK2Node_TunnelBoundary::CreateBoundariesForExpansionNodes(UEdGraphNode* Sou
 	// Create the Boundary nodes for each unique entry/exit site.
 	if (TargetGraph)
 	{
-		// Create entry node and pins
-		UK2Node_TunnelBoundary* EntryBoundaryNode = TargetGraph->CreateBlankNode<UK2Node_TunnelBoundary>();
-		MessageLog.NotifyIntermediateObjectCreation(EntryBoundaryNode, SourceNode);
-		EntryBoundaryNode->CreateNewGuid();
-		for (auto PinEntry : EntryPins)
+		if (EntryPins.Num() > 0)
 		{
-			EntryBoundaryNode->WireUpEntry(SourceNode, PinEntry.Key, PinEntry.Value, MessageLog);
+			// Create entry node and pins
+			UK2Node_TunnelBoundary* EntryBoundaryNode = TargetGraph->CreateBlankNode<UK2Node_TunnelBoundary>();
+			MessageLog.NotifyIntermediateObjectCreation(EntryBoundaryNode, SourceNode);
+			EntryBoundaryNode->CreateNewGuid();
+			for (auto PinEntry : EntryPins)
+			{
+				EntryBoundaryNode->WireUpEntry(SourceNode, PinEntry.Key, PinEntry.Value, MessageLog);
+			}
 		}
-		// Create exit node and pins
-		UK2Node_TunnelBoundary* ExitBoundaryNode = TargetGraph->CreateBlankNode<UK2Node_TunnelBoundary>();
-		MessageLog.NotifyIntermediateObjectCreation(ExitBoundaryNode, SourceNode);
-		ExitBoundaryNode->CreateNewGuid();
-		for (auto PinExit : ExitPins)
+		if (ExitPins.Num() > 0)
 		{
-			ExitBoundaryNode->WireUpExit(SourceNode, PinExit.Key, PinExit.Value, MessageLog);
+			// Create exit node and pins
+			UK2Node_TunnelBoundary* ExitBoundaryNode = TargetGraph->CreateBlankNode<UK2Node_TunnelBoundary>();
+			MessageLog.NotifyIntermediateObjectCreation(ExitBoundaryNode, SourceNode);
+			ExitBoundaryNode->CreateNewGuid();
+			for (auto PinExit : ExitPins)
+			{
+				ExitBoundaryNode->WireUpExit(SourceNode, PinExit.Key, PinExit.Value, MessageLog);
+			}
 		}
 	}
 }
@@ -266,12 +272,12 @@ void UK2Node_TunnelBoundary::WireUpTunnelExit(UK2Node_Tunnel* TunnelInstance, UE
 		UEdGraphPin* SourcePin = TunnelInstance->FindPin(TunnelPin->PinName);
 		// Wire in the node
 		UEdGraphPin* InputPin = CreatePin(EGPD_Input, UEdGraphSchema_K2::PC_Exec, TEXT(""), NULL, false, false, SourcePin->PinName);
-		MessageLog.NotifyIntermediatePinCreation(InputPin, SourcePin);
 		for (auto LinkedPin : TunnelPin->LinkedTo)
 		{
 			LinkedPin->MakeLinkTo(InputPin);
 		}
 		UEdGraphPin* OutputPin = CreatePin(EGPD_Output, UEdGraphSchema_K2::PC_Exec, TEXT(""), NULL, false, false, TEXT("TunnelExitExec"));
+		MessageLog.NotifyIntermediatePinCreation(OutputPin, SourcePin);
 		TunnelPin->BreakAllPinLinks();
 		TunnelPin->MakeLinkTo(OutputPin);
 	}
diff --git a/Engine/Source/Editor/EditorStyle/Private/EditorStyleClasses.cpp b/Engine/Source/Editor/EditorStyle/Private/EditorStyleClasses.cpp
--- a/Engine/Source/Editor/EditorStyle/Private/EditorStyleClasses.cpp
+++ b/Engine/Source/Editor/EditorStyle/Private/EditorStyleClasses.cpp
@@ -13,12 +13,14 @@ UEditorStyleSettings::UEditorStyleSettings( const FObjectInitializer& ObjectInit
 	PressedSelectionColor = FLinearColor(0.701f, 0.225f, 0.003f);
 
 	AssetEditorOpenLocation = EAssetEditorOpenLocation::Default;
+	
+	bUseGrid = true;
 
 	RegularColor = FLinearColor(0.035f, 0.035f, 0.035f);
 	RuleColor = FLinearColor(0.008f, 0.008f, 0.008f);
 	CenterColor = FLinearColor::Black;
 
-	bUseGrid = true;
+	GridSnapSize = 16.f;
 
 	bShowFriendlyNames = true;
 	LogTimestampMode = ELogTimes::None;
diff --git a/Engine/Source/Editor/EditorStyle/Private/SlateEditorStyle.cpp b/Engine/Source/Editor/EditorStyle/Private/SlateEditorStyle.cpp
--- a/Engine/Source/Editor/EditorStyle/Private/SlateEditorStyle.cpp
+++ b/Engine/Source/Editor/EditorStyle/Private/SlateEditorStyle.cpp
@@ -4101,12 +4101,6 @@ void FSlateEditorStyle::FStyle::SetupGraphEditorStyles()
 		Set( "Graph.ArrayPin.Disconnected", new IMAGE_BRUSH( "/Graph/ArrayPin_disconnected", FVector2D(11,11) ) );
 		Set( "Graph.RefPin.Connected", new IMAGE_BRUSH( "/Graph/RefPin_connected", FVector2D(11,11) ) );
 		Set( "Graph.RefPin.Disconnected", new IMAGE_BRUSH( "/Graph/RefPin_disconnected", FVector2D(11,11) ) );
-		Set( "Graph.Pin.ConnectedHovered", new IMAGE_BRUSH( "/Graph/Pin_connected", FVector2D(11,11) ) );
-		Set( "Graph.Pin.DisconnectedHovered", new IMAGE_BRUSH( "/Graph/Pin_disconnected", FVector2D(11,11) ) );
-		Set( "Graph.ArrayPin.ConnectedHovered", new IMAGE_BRUSH( "/Graph/ArrayPin_connected", FVector2D(11,11) ) );
-		Set( "Graph.ArrayPin.DisconnectedHovered", new IMAGE_BRUSH( "/Graph/ArrayPin_disconnected", FVector2D(11,11) ) );
-		Set( "Graph.RefPin.ConnectedHovered", new IMAGE_BRUSH( "/Graph/RefPin_connected", FVector2D(11,11) ) );
-		Set( "Graph.RefPin.DisconnectedHovered", new IMAGE_BRUSH( "/Graph/RefPin_disconnected", FVector2D(11,11) ) );
 
 		Set("Graph.Pin.CopyNodePinLeft_Connected", new IMAGE_BRUSH("/Graph/CopyNodePinLeft_connected", FVector2D(12, 24)));
 		Set("Graph.Pin.CopyNodePinLeft_Disconnected", new IMAGE_BRUSH("/Graph/CopyNodePinLeft_disconnected", FVector2D(12, 24)));
@@ -4117,13 +4111,9 @@ void FSlateEditorStyle::FStyle::SetupGraphEditorStyles()
 		/** Variant A Pin Styles */
 		Set( "Graph.Pin.Connected_VarA", new IMAGE_BRUSH( "/Graph/Pin_connected_VarA", FVector2D(15,11)) );
 		Set( "Graph.Pin.Disconnected_VarA", new IMAGE_BRUSH( "/Graph/Pin_disconnected_VarA", FVector2D(15,11)) );
-		Set( "Graph.Pin.ConnectedHovered_VarA", new IMAGE_BRUSH( "/Graph/Pin_connected_VarA", FVector2D(15,11) ) );
-		Set( "Graph.Pin.DisconnectedHovered_VarA", new IMAGE_BRUSH( "/Graph/Pin_disconnected_VarA", FVector2D(15,11) ) );
 
 		Set( "Graph.DelegatePin.Connected", new IMAGE_BRUSH( "/Graph/DelegatePin_Connected", FVector2D(11,11) ) );
 		Set( "Graph.DelegatePin.Disconnected", new IMAGE_BRUSH( "/Graph/DelegatePin_Disconnected", FVector2D(11,11) ) );
-		Set( "Graph.DelegatePin.ConnectedHovered", new IMAGE_BRUSH( "/Graph/DelegatePin_Connected", FVector2D(11,11) ) );
-		Set( "Graph.DelegatePin.DisconnectedHovered", new IMAGE_BRUSH( "/Graph/DelegatePin_Disconnected", FVector2D(11,11) ) );
 
 		Set( "Graph.Replication.AuthorityOnly", new IMAGE_BRUSH( "/Graph/AuthorityOnly", FVector2D(32,32) ) );
 		Set( "Graph.Replication.ClientEvent", new IMAGE_BRUSH( "/Graph/ClientEvent", FVector2D(32,32) ) );
diff --git a/Engine/Source/Editor/EditorStyle/Public/Classes/EditorStyleSettings.h b/Engine/Source/Editor/EditorStyle/Public/Classes/EditorStyleSettings.h
--- a/Engine/Source/Editor/EditorStyle/Public/Classes/EditorStyleSettings.h
+++ b/Engine/Source/Editor/EditorStyle/Public/Classes/EditorStyleSettings.h
@@ -85,17 +85,21 @@ public:
 	uint32 bUseGrid : 1;
 
 	/** The color used to represent regular grid lines */
-	UPROPERTY(EditAnywhere, config, Category = Graphs, meta = (DisplayName = "Grid regular Color"))
+	UPROPERTY(EditAnywhere, config, Category = Graphs, meta = (DisplayName = "Grid Regular Color"))
 	FLinearColor RegularColor;
 
 	/** The color used to represent ruler lines in the grid */
-	UPROPERTY(EditAnywhere, config, Category = Graphs, meta = (DisplayName = "Grid ruler Color"))
+	UPROPERTY(EditAnywhere, config, Category = Graphs, meta = (DisplayName = "Grid Ruler Color"))
 	FLinearColor RuleColor;
 
 	/** The color used to represent the center lines in the grid */
-	UPROPERTY(EditAnywhere, config, Category = Graphs, meta = (DisplayName = "Grid center Color"))
+	UPROPERTY(EditAnywhere, config, Category = Graphs, meta = (DisplayName = "Grid Center Color"))
 	FLinearColor CenterColor;
 
+	/** The custom grid snap size to use  */
+	UPROPERTY(EditAnywhere, config, Category = Graphs, meta = (ClampMin = "1.0", ClampMax = "100.0", UIMin = "1.0", UIMax = "100.0"))
+	uint32 GridSnapSize;
+
 	/** Enables animated transitions for certain menus and pop-up windows.  Note that animations may be automatically disabled at low frame rates in order to improve responsiveness. */
 	UPROPERTY(EditAnywhere, config, Category=UserInterface)
 	uint32 bEnableWindowAnimations:1;
diff --git a/Engine/Source/Editor/GameProjectGeneration/Private/GameProjectUtils.cpp b/Engine/Source/Editor/GameProjectGeneration/Private/GameProjectUtils.cpp
--- a/Engine/Source/Editor/GameProjectGeneration/Private/GameProjectUtils.cpp
+++ b/Engine/Source/Editor/GameProjectGeneration/Private/GameProjectUtils.cpp
@@ -35,6 +35,7 @@
 #include "Sound/SoundEffectSource.h"
 
 #include "PlatformInfo.h"
+#include "BlueprintSupport.h" // for FLegacyEditorOnlyBlueprintOptions::GetDefaultEditorConfig()
 
 #define LOCTEXT_NAMESPACE "GameProjectUtils"
 
@@ -1826,12 +1827,7 @@ bool GameProjectUtils::GenerateConfigFiles(const FProjectInformation& InProjectI
 	// DefaultEditor.ini
 	{
 		const FString DefaultEditorIniFilename = ProjectConfigPath / TEXT("DefaultEditor.ini");
-		FString FileContents;
-		FileContents += TEXT("[EditoronlyBP]") LINE_TERMINATOR;
-		FileContents += TEXT("bAllowClassAndBlueprintPinMatching=true") LINE_TERMINATOR;
-		FileContents += TEXT("bReplaceBlueprintWithClass=true") LINE_TERMINATOR;
-		FileContents += TEXT("bDontLoadBlueprintOutsideEditor=true") LINE_TERMINATOR;
-		FileContents += TEXT("bBlueprintIsNotBlueprintType=true") LINE_TERMINATOR;
+		FString FileContents = FLegacyEditorOnlyBlueprintOptions::GetDefaultEditorConfig();
 
 		if (WriteOutputFile(DefaultEditorIniFilename, FileContents, OutFailReason))
 		{
diff --git a/Engine/Source/Editor/GameplayTagsEditor/Classes/GameplayTagsK2Node_SwitchGameplayTag.h b/Engine/Source/Editor/GameplayTagsEditor/Classes/GameplayTagsK2Node_SwitchGameplayTag.h
--- a/Engine/Source/Editor/GameplayTagsEditor/Classes/GameplayTagsK2Node_SwitchGameplayTag.h
+++ b/Engine/Source/Editor/GameplayTagsEditor/Classes/GameplayTagsK2Node_SwitchGameplayTag.h
@@ -41,6 +41,7 @@ class UGameplayTagsK2Node_SwitchGameplayTag : public UK2Node_Switch
 	GAMEPLAYTAGSEDITOR_API virtual void AddPinToSwitchNode() override;
 	virtual FString GetUniquePinName() override;
 	virtual FEdGraphPinType GetPinType() const override;
+	GAMEPLAYTAGSEDITOR_API virtual FEdGraphPinType GetInnerCaseType() const override;
 	// End of UK2Node_Switch Interface
 
 	virtual FString GetPinNameGivenIndex(int32 Index) override;
diff --git a/Engine/Source/Editor/GameplayTagsEditor/Private/GameplayTagsK2Node_SwitchGameplayTag.cpp b/Engine/Source/Editor/GameplayTagsEditor/Private/GameplayTagsK2Node_SwitchGameplayTag.cpp
--- a/Engine/Source/Editor/GameplayTagsEditor/Private/GameplayTagsK2Node_SwitchGameplayTag.cpp
+++ b/Engine/Source/Editor/GameplayTagsEditor/Private/GameplayTagsK2Node_SwitchGameplayTag.cpp
@@ -111,6 +111,14 @@ FEdGraphPinType UGameplayTagsK2Node_SwitchGameplayTag::GetPinType() const
 	return PinType;
 }
 
+FEdGraphPinType UGameplayTagsK2Node_SwitchGameplayTag::GetInnerCaseType() const
+{
+	// This type should match the second argument of UBlueprintGameplayTagLibrary::NotEqual_TagTag !
+	FEdGraphPinType PinType;
+	PinType.PinCategory = UEdGraphSchema_K2::PC_String;
+	return PinType;
+}
+
 FString UGameplayTagsK2Node_SwitchGameplayTag::GetPinNameGivenIndex(int32 Index)
 {
 	check(Index);
diff --git a/Engine/Source/Editor/GameplayTasksEditor/Private/K2Node_LatentGameplayTaskCall.cpp b/Engine/Source/Editor/GameplayTasksEditor/Private/K2Node_LatentGameplayTaskCall.cpp
--- a/Engine/Source/Editor/GameplayTasksEditor/Private/K2Node_LatentGameplayTaskCall.cpp
+++ b/Engine/Source/Editor/GameplayTasksEditor/Private/K2Node_LatentGameplayTaskCall.cpp
@@ -123,6 +123,7 @@ void UK2Node_LatentGameplayTaskCall::ReallocatePinsDuringReconstruction(TArray<U
 	{
 		CreatePinsForClass(UseSpawnClass);
 	}
+	RestoreSplitPins(OldPins);
 }
 
 UEdGraphPin* UK2Node_LatentGameplayTaskCall::GetClassPin(const TArray<UEdGraphPin*>* InPinsToSearch /*= NULL*/) const
diff --git a/Engine/Source/Editor/GraphEditor/Private/AnimationPins/SGraphPinPose.h b/Engine/Source/Editor/GraphEditor/Private/AnimationPins/SGraphPinPose.h
--- a/Engine/Source/Editor/GraphEditor/Private/AnimationPins/SGraphPinPose.h
+++ b/Engine/Source/Editor/GraphEditor/Private/AnimationPins/SGraphPinPose.h
@@ -14,4 +14,7 @@ protected:
 	//~ Begin SGraphPin Interface
 	virtual const FSlateBrush* GetPinIcon() const override;
 	//~ End SGraphPin Interface
+
+	mutable const FSlateBrush* CachedImg_Pin_ConnectedHovered;
+	mutable const FSlateBrush* CachedImg_Pin_DisconnectedHovered;
 };
diff --git a/Engine/Source/Editor/GraphEditor/Private/KismetNodes/SGraphNodeK2CreateDelegate.cpp b/Engine/Source/Editor/GraphEditor/Private/KismetNodes/SGraphNodeK2CreateDelegate.cpp
--- a/Engine/Source/Editor/GraphEditor/Private/KismetNodes/SGraphNodeK2CreateDelegate.cpp
+++ b/Engine/Source/Editor/GraphEditor/Private/KismetNodes/SGraphNodeK2CreateDelegate.cpp
@@ -4,7 +4,7 @@
 #include "SGraphNodeK2Base.h"
 #include "SGraphNodeK2CreateDelegate.h"
 
-FString SGraphNodeK2CreateDelegate::FunctionDescription(const UFunction* Function)
+FString SGraphNodeK2CreateDelegate::FunctionDescription(const UFunction* Function, const bool bOnlyDescribeSignature, const int32 CharacterLimit)
 {
 	if(!Function || !Function->GetOuter())
 	{
@@ -16,7 +16,7 @@ FString SGraphNodeK2CreateDelegate::FunctionDescription(const UFunction* Functio
 	//FString Result = Function->GetOuter()->GetName() + TEXT("::") + Function->GetName();
 	//Result += TEXT("(");
 
-	FString Result = Function->GetName() + TEXT("(");
+	FString Result = (bOnlyDescribeSignature ? TEXT("") : Function->GetName()) + TEXT("(");
 	bool bFirst = true;
 	for (TFieldIterator<UProperty> PropIt(Function); PropIt && (PropIt->PropertyFlags & CPF_Parm); ++PropIt)
 	{
@@ -28,12 +28,12 @@ FString SGraphNodeK2CreateDelegate::FunctionDescription(const UFunction* Functio
 			{
 				Result += TEXT(", ");
 			}
-			if(Result.Len() > 32)
+			if(CharacterLimit > INDEX_NONE && Result.Len() > CharacterLimit)
 			{
 				Result += TEXT("...");
 				break;
 			}
-			Result += Param->GetName();
+			Result += bOnlyDescribeSignature ? UEdGraphSchema_K2::TypeToText(Param).ToString() : Param->GetName();
 			bFirst = false;
 		}
 	}
@@ -99,7 +99,7 @@ void SGraphNodeK2CreateDelegate::OnFunctionSelected(TSharedPtr<FFunctionItemData
 	}
 }
 
-void SGraphNodeK2CreateDelegate::CreateBelowWidgetControls(TSharedPtr<SVerticalBox> MainBox)
+void SGraphNodeK2CreateDelegate::CreateBelowPinControls(TSharedPtr<SVerticalBox> MainBox)
 {
 	if(UK2Node_CreateDelegate* Node = Cast<UK2Node_CreateDelegate>(GraphNode))
 	{
@@ -108,6 +108,23 @@ void SGraphNodeK2CreateDelegate::CreateBelowWidgetControls(TSharedPtr<SVerticalB
 
 		if(FunctionSignature && ScopeClass)
 		{
+			FText FunctionSignaturePrompt;
+			{
+				FFormatNamedArguments FormatArguments;
+				FormatArguments.Add(TEXT("FunctionSignature"), FText::FromString(FunctionDescription(FunctionSignature, true)));
+				FunctionSignaturePrompt = FText::Format(NSLOCTEXT("GraphNodeK2Create", "FunctionSignaturePrompt", "Signature: {FunctionSignature}"), FormatArguments);
+			}
+
+			MainBox->AddSlot()
+				.AutoHeight()
+				.VAlign(VAlign_Fill)
+				.Padding(4.0f)
+				[
+					SNew(STextBlock)
+					.Text(FunctionSignaturePrompt)
+					.ToolTipText(FText::FromString(FunctionDescription(FunctionSignature, true, INDEX_NONE)))
+				];
+
 			FunctionDataItems.Empty();
 
 			for(TFieldIterator<UFunction> It(ScopeClass); It; ++It)
@@ -149,6 +166,7 @@ void SGraphNodeK2CreateDelegate::CreateBelowWidgetControls(TSharedPtr<SVerticalB
 			MainBox->AddSlot()
 				.AutoHeight()
 				.VAlign(VAlign_Fill)
+				.Padding(4.0f)
 				[
 					SelectFunctionWidgetRef
 				];
diff --git a/Engine/Source/Editor/GraphEditor/Private/KismetNodes/SGraphNodeK2CreateDelegate.h b/Engine/Source/Editor/GraphEditor/Private/KismetNodes/SGraphNodeK2CreateDelegate.h
--- a/Engine/Source/Editor/GraphEditor/Private/KismetNodes/SGraphNodeK2CreateDelegate.h
+++ b/Engine/Source/Editor/GraphEditor/Private/KismetNodes/SGraphNodeK2CreateDelegate.h
@@ -16,13 +16,13 @@ public:
 	TArray<TSharedPtr<FFunctionItemData>> FunctionDataItems;
 	TWeakPtr<SComboButton> SelectFunctionWidget;
 public:
-	static FString FunctionDescription(const UFunction* Function);
-
 	virtual ~SGraphNodeK2CreateDelegate();
 	void Construct(const FArguments& InArgs, UK2Node* InNode);
-	virtual void CreateBelowWidgetControls(TSharedPtr<SVerticalBox> MainBox) override;
+	virtual void CreateBelowPinControls(TSharedPtr<SVerticalBox> MainBox) override;
 
 protected:
+	static FString FunctionDescription(const UFunction* Function, const bool bOnlyDescribeSignature = false, const int32 CharacterLimit = 32);
+
 	FText GetCurrentFunctionDescription() const;
 	TSharedRef<ITableRow> HandleGenerateRowFunction(TSharedPtr<FFunctionItemData> FunctionItemData, const TSharedRef<STableViewBase>& OwnerTable);
 	void OnFunctionSelected(TSharedPtr<FFunctionItemData> FunctionItemData, ESelectInfo::Type SelectInfo);
diff --git a/Engine/Source/Editor/GraphEditor/Private/KismetNodes/SGraphNodeSwitchStatement.cpp b/Engine/Source/Editor/GraphEditor/Private/KismetNodes/SGraphNodeSwitchStatement.cpp
--- a/Engine/Source/Editor/GraphEditor/Private/KismetNodes/SGraphNodeSwitchStatement.cpp
+++ b/Engine/Source/Editor/GraphEditor/Private/KismetNodes/SGraphNodeSwitchStatement.cpp
@@ -24,7 +24,6 @@ public:
 		SGraphPin::Construct(SGraphPin::FArguments().PinLabelStyle(FName("Graph.Node.DefaultPinName")), InPin);
 
 		bWasEventPin = false;
-		CachePinBrushes(/*bForceCache=*/ true);
 	}
 };
 
diff --git a/Engine/Source/Editor/GraphEditor/Private/KismetPins/SGraphPinEnum.cpp b/Engine/Source/Editor/GraphEditor/Private/KismetPins/SGraphPinEnum.cpp
--- a/Engine/Source/Editor/GraphEditor/Private/KismetPins/SGraphPinEnum.cpp
+++ b/Engine/Source/Editor/GraphEditor/Private/KismetPins/SGraphPinEnum.cpp
@@ -136,9 +136,18 @@ void SGraphPinEnum::ComboBoxSelectionChanged( TSharedPtr<int32> NewSelection, ES
 {
 	UEnum* EnumPtr = Cast<UEnum>(GraphPinObj->PinType.PinSubCategoryObject.Get());
 	check(EnumPtr);
-	check(*NewSelection < EnumPtr->NumEnums() - 1);
 
-	FString EnumSelectionString = EnumPtr->GetEnumName(*NewSelection);
+	FString EnumSelectionString;
+	if (NewSelection.IsValid())
+	{
+		check(*NewSelection < EnumPtr->NumEnums() - 1);
+		EnumSelectionString  = EnumPtr->GetEnumName(*NewSelection);
+	}
+	else
+	{
+		EnumSelectionString = FName(NAME_None).ToString();
+	}
+	
 
 	if(GraphPinObj->GetDefaultAsString() != EnumSelectionString)
 	{
diff --git a/Engine/Source/Editor/GraphEditor/Private/KismetPins/SGraphPinExec.cpp b/Engine/Source/Editor/GraphEditor/Private/KismetPins/SGraphPinExec.cpp
--- a/Engine/Source/Editor/GraphEditor/Private/KismetPins/SGraphPinExec.cpp
+++ b/Engine/Source/Editor/GraphEditor/Private/KismetPins/SGraphPinExec.cpp
@@ -9,7 +9,11 @@ void SGraphPinExec::Construct(const FArguments& InArgs, UEdGraphPin* InPin)
 	SGraphPin::Construct(SGraphPin::FArguments(), InPin);
 
 	bWasEventPin = false;
-	CachePinBrushes(/*bForceCache=*/ true);
+
+	CachedImg_Pin_ConnectedHovered = FEditorStyle::GetBrush(TEXT("Graph.ExecPin.ConnectedHovered"));
+	CachedImg_Pin_Connected = FEditorStyle::GetBrush(TEXT("Graph.ExecPin.Connected"));
+	CachedImg_Pin_DisconnectedHovered = FEditorStyle::GetBrush(TEXT("Graph.ExecPin.DisconnectedHovered"));
+	CachedImg_Pin_Disconnected = FEditorStyle::GetBrush(TEXT("Graph.ExecPin.Disconnected"));
 }
 
 TSharedRef<SWidget>	SGraphPinExec::GetDefaultValueWidget()
@@ -19,8 +23,6 @@ TSharedRef<SWidget>	SGraphPinExec::GetDefaultValueWidget()
 
 const FSlateBrush* SGraphPinExec::GetPinIcon() const
 {
-	CachePinBrushes();
-
 	const FSlateBrush* Brush = NULL;
 
 	if (IsConnected())
@@ -34,14 +36,3 @@ const FSlateBrush* SGraphPinExec::GetPinIcon() const
 
 	return Brush;
 }
-
-void SGraphPinExec::CachePinBrushes(bool bForceCache) const
-{
-	if (bForceCache)
-	{
-		CachedImg_Pin_ConnectedHovered = FEditorStyle::GetBrush(TEXT("Graph.ExecPin.ConnectedHovered"));
-		CachedImg_Pin_Connected = FEditorStyle::GetBrush(TEXT("Graph.ExecPin.Connected"));
-		CachedImg_Pin_DisconnectedHovered = FEditorStyle::GetBrush(TEXT("Graph.ExecPin.DisconnectedHovered"));
-		CachedImg_Pin_Disconnected = FEditorStyle::GetBrush(TEXT("Graph.ExecPin.Disconnected"));
-	}
-}
\ No newline at end of file
diff --git a/Engine/Source/Editor/GraphEditor/Private/KismetPins/SGraphPinExec.h b/Engine/Source/Editor/GraphEditor/Private/KismetPins/SGraphPinExec.h
--- a/Engine/Source/Editor/GraphEditor/Private/KismetPins/SGraphPinExec.h
+++ b/Engine/Source/Editor/GraphEditor/Private/KismetPins/SGraphPinExec.h
@@ -16,8 +16,9 @@ protected:
 	virtual const FSlateBrush* GetPinIcon() const override;
 	//~ End SGraphPin Interface
 
-	void CachePinBrushes(bool bForceCache = false) const;
-
 protected:
 	mutable bool bWasEventPin;
+
+	mutable const FSlateBrush* CachedImg_Pin_ConnectedHovered;
+	mutable const FSlateBrush* CachedImg_Pin_DisconnectedHovered;
 };
diff --git a/Engine/Source/Editor/GraphEditor/Private/SGraphActionMenu.cpp b/Engine/Source/Editor/GraphEditor/Private/SGraphActionMenu.cpp
--- a/Engine/Source/Editor/GraphEditor/Private/SGraphActionMenu.cpp
+++ b/Engine/Source/Editor/GraphEditor/Private/SGraphActionMenu.cpp
@@ -770,10 +770,11 @@ int32 SGraphActionMenu::GetActionFilteredWeight( const FGraphActionListBuilderBa
 	int32 TotalWeight = 0;
 
 	// Some simple weight figures to help find the most appropriate match	
-	int32 WholeMatchWeightMultiplier = 1;
-	int32 DescriptionWeight = 1;
-	int32 CategoryWeight = 1;
-	int32 NodeTitleWeight = 1;
+	const int32 WholeMatchWeightMultiplier = 2;
+	const int32 DescriptionWeight = 10;
+	const int32 CategoryWeight = 1;
+	const int32 NodeTitleWeight = 1;
+	const int32 KeywordWeight = 4;
 
 	// Helper array
 	struct FArrayWithWeight
@@ -800,7 +801,7 @@ int32 SGraphActionMenu::GetActionFilteredWeight( const FGraphActionListBuilderBa
 		const FString& SearchText = InCurrentAction.GetSearchTextForFirstAction();
 
 		// First the keywords
-		WeightedArrayList.Add(FArrayWithWeight(&InCurrentAction.GetSearchKeywordsArrayForFirstAction(), 10));
+		WeightedArrayList.Add(FArrayWithWeight(&InCurrentAction.GetSearchKeywordsArrayForFirstAction(), KeywordWeight));
 
 		// The description
 		WeightedArrayList.Add(FArrayWithWeight(&InCurrentAction.GetMenuDescriptionArrayForFirstAction(), DescriptionWeight));
@@ -812,17 +813,17 @@ int32 SGraphActionMenu::GetActionFilteredWeight( const FGraphActionListBuilderBa
 		WeightedArrayList.Add(FArrayWithWeight(&InCurrentAction.GetSearchCategoryArrayForFirstAction(), CategoryWeight));
 
 		// Now iterate through all the filter terms and calculate a 'weight' using the values and multipliers
-		FString EachTerm;
-		FString EachTermSanitized;
+		const FString* EachTerm = nullptr;
+		const FString* EachTermSanitized = nullptr;
 		for (int32 FilterIndex = 0; FilterIndex < InFilterTerms.Num(); ++FilterIndex)
 		{
-			EachTerm = InFilterTerms[FilterIndex];
-			EachTermSanitized = InSanitizedFilterTerms[FilterIndex];
-			if( SearchText.Contains( EachTerm, ESearchCase::CaseSensitive ) )
+			EachTerm = &InFilterTerms[FilterIndex];
+			EachTermSanitized = &InSanitizedFilterTerms[FilterIndex];
+			if( SearchText.Contains( *EachTerm, ESearchCase::CaseSensitive ) )
 			{
 				TotalWeight += 2;
 			}
-			else if (SearchText.Contains(EachTermSanitized, ESearchCase::CaseSensitive))
+			else if (SearchText.Contains(*EachTermSanitized, ESearchCase::CaseSensitive))
 			{
 				TotalWeight++;
 			}		
@@ -836,21 +837,35 @@ int32 SGraphActionMenu::GetActionFilteredWeight( const FGraphActionListBuilderBa
 				for (int32 iEachWord = 0; iEachWord < KeywordArray.Num() ; iEachWord++)
 				{
 					// If we get an exact match weight the find count to get exact matches higher priority
-					if (KeywordArray[iEachWord] == EachTerm)
+					if (KeywordArray[iEachWord].StartsWith(*EachTerm, ESearchCase::CaseSensitive))
 					{
-						WeightPerList += EachWeight * WholeMatchWeightMultiplier;
+						if (iEachWord == 0)
+						{
+							WeightPerList += EachWeight * WholeMatchWeightMultiplier;
+						}
+						else
+						{
+							WeightPerList += EachWeight;
+						}
 						WholeMatchCount++;
 					}
-					else if (KeywordArray[iEachWord].Contains(EachTerm, ESearchCase::CaseSensitive))
+					else if (KeywordArray[iEachWord].Contains(*EachTerm, ESearchCase::CaseSensitive))
 					{
 						WeightPerList += EachWeight;
 					}
-					else if (KeywordArray[iEachWord] == EachTermSanitized)
+					if (KeywordArray[iEachWord].StartsWith(*EachTermSanitized, ESearchCase::CaseSensitive))
 					{
-						WeightPerList += ( EachWeight * WholeMatchWeightMultiplier ) / 2;
+						if (iEachWord == 0)
+						{
+							WeightPerList += EachWeight * WholeMatchWeightMultiplier;
+						}
+						else
+						{
+							WeightPerList += EachWeight;
+						}
 						WholeMatchCount++;
 					}
-					else if (KeywordArray[iEachWord].Contains(EachTermSanitized, ESearchCase::CaseSensitive))
+					else if (KeywordArray[iEachWord].Contains(*EachTermSanitized, ESearchCase::CaseSensitive))
 					{
 						WeightPerList += EachWeight / 2;
 					}
diff --git a/Engine/Source/Editor/GraphEditor/Private/SGraphPin.cpp b/Engine/Source/Editor/GraphEditor/Private/SGraphPin.cpp
--- a/Engine/Source/Editor/GraphEditor/Private/SGraphPin.cpp
+++ b/Engine/Source/Editor/GraphEditor/Private/SGraphPin.cpp
@@ -6,66 +6,9 @@
 #include "Editor/UnrealEd/Public/Kismet2/KismetDebugUtilities.h"
 #include "Editor/UnrealEd/Public/DragAndDrop/AssetDragDropOp.h"
 #include "BlueprintUtilities.h"
-#include "SLevelOfDetailBranchNode.h"
 #include "ScopedTransaction.h"
-
-/////////////////////////////////////////////////////
-// FKnotNetCollector
-
-struct FKnotNetCollector
-{
-	TSet<UEdGraphNode*> VisitedNodes;
-	TSet<UEdGraphPin*> VisitedPins;
-
-	FKnotNetCollector(UEdGraphPin* StartingPin)
-	{
-		TraversePin(StartingPin);
-	}
-
-	void TraversePin(UEdGraphPin* Pin)
-	{
-		if (UK2Node_Knot* Knot = Cast<UK2Node_Knot>(Pin->GetOwningNodeUnchecked()))
-		{
-			TraverseNodes(Knot);
-		}
-		else
-		{
-			VisitedPins.Add(Pin);
-
-			for (UEdGraphPin* OtherPin : Pin->LinkedTo)
-			{
-				UEdGraphNode* OtherNode = OtherPin->GetOwningNodeUnchecked();
-				if (OtherNode && OtherNode->IsA(UK2Node_Knot::StaticClass()))
-				{
-					TraverseNodes(OtherNode);
-				}
-			}
-		}
-	}
-
-	void TraverseNodes(UEdGraphNode* Node)
-	{
-		if (VisitedNodes.Contains(Node))
-		{
-			return;
-		}
-		VisitedNodes.Add(Node);
-
-		for (UEdGraphPin* MyPin : Node->Pins)
-		{
-			VisitedPins.Add(MyPin);
-
-			for (UEdGraphPin* OtherPin : MyPin->LinkedTo)
-			{
-				UEdGraphNode* OtherNode = OtherPin->GetOwningNodeUnchecked();
-				if (OtherNode && OtherNode->IsA(UK2Node_Knot::StaticClass()))
-				{
-					TraverseNodes(OtherNode);
-				}
-			}
-		}
-	}
-};
+#include "SLevelOfDetailBranchNode.h"
+#include "SPinTypeSelector.h"
 
 /////////////////////////////////////////////////////
 // FGraphPinHandle
@@ -124,32 +67,26 @@ SGraphPin::SGraphPin()
 	// Make these names const so they're not created for every pin
 
 	/** Original Pin Styles */
-	static const FName NAME_Pin_ConnectedHovered("Graph.Pin.ConnectedHovered");
 	static const FName NAME_Pin_Connected("Graph.Pin.Connected");
-	static const FName NAME_Pin_DisconnectedHovered("Graph.Pin.DisconnectedHovered");
 	static const FName NAME_Pin_Disconnected("Graph.Pin.Disconnected");
 
 	/** Variant A Pin Styles */
-	static const FName NAME_Pin_ConnectedHovered_VarA("Graph.Pin.ConnectedHovered_VarA");
 	static const FName NAME_Pin_Connected_VarA("Graph.Pin.Connected_VarA");
-	static const FName NAME_Pin_DisconnectedHovered_VarA("Graph.Pin.DisconnectedHovered_VarA");
 	static const FName NAME_Pin_Disconnected_VarA("Graph.Pin.Disconnected_VarA");
 
-	static const FName NAME_ArrayPin_ConnectedHovered("Graph.ArrayPin.ConnectedHovered");
 	static const FName NAME_ArrayPin_Connected("Graph.ArrayPin.Connected");
-	static const FName NAME_ArrayPin_DisconnectedHovered("Graph.ArrayPin.DisconnectedHovered");
 	static const FName NAME_ArrayPin_Disconnected("Graph.ArrayPin.Disconnected");
 
-	static const FName NAME_RefPin_ConnectedHovered("Graph.RefPin.ConnectedHovered");
 	static const FName NAME_RefPin_Connected("Graph.RefPin.Connected");
-	static const FName NAME_RefPin_DisconnectedHovered("Graph.RefPin.DisconnectedHovered");
 	static const FName NAME_RefPin_Disconnected("Graph.RefPin.Disconnected");
 
-	static const FName NAME_DelegatePin_ConnectedHovered("Graph.DelegatePin.ConnectedHovered");
 	static const FName NAME_DelegatePin_Connected("Graph.DelegatePin.Connected");
-	static const FName NAME_DelegatePin_DisconnectedHovered("Graph.DelegatePin.DisconnectedHovered");
 	static const FName NAME_DelegatePin_Disconnected("Graph.DelegatePin.Disconnected");
 
+	static const FName NAME_SetPin("Kismet.VariableList.SetTypeIcon");
+	static const FName NAME_MapPinKey("Kismet.VariableList.MapKeyTypeIcon");
+	static const FName NAME_MapPinValue("Kismet.VariableList.MapValueTypeIcon");
+
 	static const FName NAME_Pin_Background("Graph.Pin.Background");
 	static const FName NAME_Pin_BackgroundHovered("Graph.Pin.BackgroundHovered");
 
@@ -158,35 +95,29 @@ SGraphPin::SGraphPin()
 	switch(StyleType)
 	{
 	case BPST_VariantA:
-		CachedImg_Pin_ConnectedHovered = FEditorStyle::GetBrush( NAME_Pin_ConnectedHovered_VarA );
 		CachedImg_Pin_Connected = FEditorStyle::GetBrush( NAME_Pin_Connected_VarA );
-		CachedImg_Pin_DisconnectedHovered = FEditorStyle::GetBrush( NAME_Pin_DisconnectedHovered_VarA );
 		CachedImg_Pin_Disconnected = FEditorStyle::GetBrush( NAME_Pin_Disconnected_VarA );
 		break;
 	case BPST_Original:
 	default:
-		CachedImg_Pin_ConnectedHovered = FEditorStyle::GetBrush( NAME_Pin_ConnectedHovered );
 		CachedImg_Pin_Connected = FEditorStyle::GetBrush( NAME_Pin_Connected );
-		CachedImg_Pin_DisconnectedHovered = FEditorStyle::GetBrush( NAME_Pin_DisconnectedHovered );
 		CachedImg_Pin_Disconnected = FEditorStyle::GetBrush( NAME_Pin_Disconnected );
 		break;
 	}
 
-	CachedImg_RefPin_ConnectedHovered = FEditorStyle::GetBrush( NAME_RefPin_ConnectedHovered );
 	CachedImg_RefPin_Connected = FEditorStyle::GetBrush( NAME_RefPin_Connected );
-	CachedImg_RefPin_DisconnectedHovered = FEditorStyle::GetBrush( NAME_RefPin_DisconnectedHovered );
 	CachedImg_RefPin_Disconnected = FEditorStyle::GetBrush( NAME_RefPin_Disconnected );
 
-	CachedImg_ArrayPin_ConnectedHovered = FEditorStyle::GetBrush( NAME_ArrayPin_ConnectedHovered );
 	CachedImg_ArrayPin_Connected = FEditorStyle::GetBrush( NAME_ArrayPin_Connected );
-	CachedImg_ArrayPin_DisconnectedHovered = FEditorStyle::GetBrush( NAME_ArrayPin_DisconnectedHovered );
 	CachedImg_ArrayPin_Disconnected = FEditorStyle::GetBrush( NAME_ArrayPin_Disconnected );
 
-	CachedImg_DelegatePin_ConnectedHovered = FEditorStyle::GetBrush( NAME_DelegatePin_ConnectedHovered );
 	CachedImg_DelegatePin_Connected = FEditorStyle::GetBrush( NAME_DelegatePin_Connected );
-	CachedImg_DelegatePin_DisconnectedHovered = FEditorStyle::GetBrush( NAME_DelegatePin_DisconnectedHovered );
 	CachedImg_DelegatePin_Disconnected = FEditorStyle::GetBrush( NAME_DelegatePin_Disconnected );
 
+	CachedImg_SetPin = FEditorStyle::GetBrush(NAME_SetPin);
+	CachedImg_MapPinKey = FEditorStyle::GetBrush(NAME_MapPinKey);
+	CachedImg_MapPinValue = FEditorStyle::GetBrush(NAME_MapPinValue);
+
 	CachedImg_Pin_Background = FEditorStyle::GetBrush( NAME_Pin_Background );
 	CachedImg_Pin_BackgroundHovered = FEditorStyle::GetBrush( NAME_Pin_BackgroundHovered );
 }
@@ -213,13 +144,11 @@ void SGraphPin::Construct(const FArguments& InArgs, UEdGraphPin* InPin)
 	const bool bIsInput = (GetDirection() == EGPD_Input);
 
 	// Create the pin icon widget
-	TSharedRef<SWidget> ActualPinWidget =
-		SAssignNew(PinImage, SImage)
-		.Image(this, &SGraphPin::GetPinIcon)
-		.IsEnabled(this, &SGraphPin::GetIsConnectable)
-		.ColorAndOpacity(this, &SGraphPin::GetPinColor)
-		.OnMouseButtonDown(this, &SGraphPin::OnPinMouseDown)
-		.Cursor(this, &SGraphPin::GetPinCursor);
+	TSharedRef<SWidget> ActualPinWidget = SPinTypeSelector::ConstructPinTypeImage(
+		TAttribute<const FSlateBrush*>::Create( TAttribute<const FSlateBrush*>::FGetter::CreateRaw(this, &SGraphPin::GetPinIcon ) ),
+		TAttribute<FSlateColor>::Create( TAttribute<FSlateColor>::FGetter::CreateRaw(this, &SGraphPin::GetPinColor) ),
+		TAttribute<const FSlateBrush*>::Create( TAttribute<const FSlateBrush*>::FGetter::CreateRaw(this, &SGraphPin::GetSecondaryPinIcon ) ),
+		TAttribute<FSlateColor>::Create( TAttribute<FSlateColor>::FGetter::CreateRaw(this, &SGraphPin::GetSecondaryPinColor) ));
 
 	// Create the pin indicator widget (used for watched values)
 	static const FName NAME_NoBorder("NoBorder");
@@ -602,13 +531,61 @@ void SGraphPin::OnMouseEnter( const FGeometry& MyGeometry, const FPointerEvent&
 {
 	if (!bIsHovered)
 	{
-		FKnotNetCollector NetCollector(GetPinObj());
-
-		TSharedPtr<SGraphPanel> Panel = OwnerNodePtr.Pin()->GetOwnerPanel();
-		for (UEdGraphPin* PinInNet : NetCollector.VisitedPins)
+		UEdGraphPin* MyPin = GetPinObj();
+		if (MyPin && !MyPin->IsPendingKill() && MyPin->GetOuter() && MyPin->GetOuter()->IsA(UEdGraphNode::StaticClass()))
 		{
-			Panel->AddPinToHoverSet(PinInNet);
-			HoverPinSet.Add(PinInNet);
+			struct FHoverPinHelper
+			{
+				FHoverPinHelper(TSharedPtr<SGraphPanel> Panel, TSet<FEdGraphPinReference>& PinSet)
+					: PinSetOut(PinSet), TargetPanel(Panel)
+				{}
+
+				void SetHovered(UEdGraphPin* Pin)
+				{
+					bool bAlreadyAdded = false;
+					PinSetOut.Add(Pin, &bAlreadyAdded);
+					if (bAlreadyAdded)
+					{
+						return;
+					}
+					TargetPanel->AddPinToHoverSet(Pin);
+	
+					for (UEdGraphPin* LinkedPin : Pin->LinkedTo)
+					{
+						if (UK2Node_Knot* OwningNode = Cast<UK2Node_Knot>(LinkedPin->GetOwningNodeUnchecked()))
+						{
+							SetHovered(OwningNode);
+						}
+					}
+				}
+
+			private:
+				void SetHovered(UK2Node_Knot* KnotNode)
+				{
+					bool bAlreadyTraversed = false;
+					IntermediateNodes.Add(KnotNode, &bAlreadyTraversed);
+
+					if (!bAlreadyTraversed)
+					{
+						for (UEdGraphPin* KnotPin : KnotNode->Pins)
+						{
+							SetHovered(KnotPin);
+						}
+					}
+				}
+
+			private:
+				TSet<UK2Node_Knot*> IntermediateNodes;
+				TSet<FEdGraphPinReference>& PinSetOut;
+				TSharedPtr<SGraphPanel>   TargetPanel;
+			};
+
+
+			TSharedPtr<SGraphPanel> Panel = OwnerNodePtr.Pin()->GetOwnerPanel();
+			if (Panel.IsValid())
+			{
+				FHoverPinHelper(Panel, HoverPinSet).SetHovered(MyPin);
+			}
 		}
 	}
 
@@ -827,12 +804,21 @@ EEdGraphPinDirection SGraphPin::GetDirection() const
 	return static_cast<EEdGraphPinDirection>(GraphPinObj->Direction);
 }
 
-/** @return whether this pin is an array value */
 bool SGraphPin::IsArray() const
 {
 	return GraphPinObj->PinType.bIsArray;
 }
 
+bool SGraphPin::IsSet() const
+{
+	return GraphPinObj->PinType.bIsSet;
+}
+
+bool SGraphPin::IsMap() const
+{
+	return GraphPinObj->PinType.bIsMap;
+}
+
 bool SGraphPin::IsByRef() const
 {
 	return GraphPinObj->PinType.bIsReference;
@@ -867,49 +853,66 @@ const FSlateBrush* SGraphPin::GetPinIcon() const
 	{
 		if (IsConnected())
 		{
-			return IsHovered() ? CachedImg_ArrayPin_ConnectedHovered : CachedImg_ArrayPin_Connected;
+			return CachedImg_ArrayPin_Connected;
 		}
 		else
 		{
-			return IsHovered() ? CachedImg_ArrayPin_DisconnectedHovered : CachedImg_ArrayPin_Disconnected;
+			return CachedImg_ArrayPin_Disconnected;
 		}
 	}
 	else if(IsDelegate())
 	{
 		if (IsConnected())
 		{
-			return IsHovered() ? CachedImg_DelegatePin_ConnectedHovered : CachedImg_DelegatePin_Connected;		
+			return CachedImg_DelegatePin_Connected;		
 		}
 		else
 		{
-			return IsHovered() ? CachedImg_DelegatePin_DisconnectedHovered : CachedImg_DelegatePin_Disconnected;
+			return CachedImg_DelegatePin_Disconnected;
 		}
 	}
 	else if (GraphPinObj->bDisplayAsMutableRef || IsByMutableRef())
 	{
 		if (IsConnected())
 		{
-			return IsHovered() ? CachedImg_RefPin_ConnectedHovered : CachedImg_RefPin_Connected;		
+			return CachedImg_RefPin_Connected;		
 		}
 		else
 		{
-			return IsHovered() ? CachedImg_RefPin_DisconnectedHovered : CachedImg_RefPin_Disconnected;
+			return CachedImg_RefPin_Disconnected;
 		}
 
 	}
+	else if (IsSet())
+	{
+		return CachedImg_SetPin;
+	}
+	else if (IsMap())
+	{
+		return CachedImg_MapPinKey;
+	}
 	else
 	{
 		if (IsConnected())
 		{
-			return IsHovered() ? CachedImg_Pin_ConnectedHovered : CachedImg_Pin_Connected;
+			return CachedImg_Pin_Connected;
 		}
 		else
 		{
-			return IsHovered() ? CachedImg_Pin_DisconnectedHovered : CachedImg_Pin_Disconnected;
+			return CachedImg_Pin_Disconnected;
 		}
 	}
 }
 
+const FSlateBrush* SGraphPin::GetSecondaryPinIcon() const
+{
+	if( !GraphPinObj->IsPendingKill() && GraphPinObj->PinType.bIsMap )
+	{
+		return CachedImg_MapPinValue;
+	}
+	return nullptr;
+}
+
 const FSlateBrush* SGraphPin::GetPinBorder() const
 {
 	bool bIsMarkedPin = false;
@@ -946,6 +949,12 @@ FSlateColor SGraphPin::GetPinColor() const
 	return FLinearColor::White;
 }
 
+FSlateColor SGraphPin::GetSecondaryPinColor() const
+{
+	const UEdGraphSchema_K2* Schema = Cast<UEdGraphSchema_K2>(!GraphPinObj->IsPendingKill() ? GraphPinObj->GetSchema() : nullptr);
+	return Schema ? Schema->GetSecondaryPinTypeColor(GraphPinObj->PinType) : FLinearColor::White;
+}
+
 FSlateColor SGraphPin::GetPinTextColor() const
 {
 	// If there is no schema there is no owning node (or basically this is a deleted node)
diff --git a/Engine/Source/Editor/GraphEditor/Private/SNodePanel.cpp b/Engine/Source/Editor/GraphEditor/Private/SNodePanel.cpp
--- a/Engine/Source/Editor/GraphEditor/Private/SNodePanel.cpp
+++ b/Engine/Source/Editor/GraphEditor/Private/SNodePanel.cpp
@@ -707,8 +707,8 @@ FReply SNodePanel::OnMouseMove(const FGeometry& MyGeometry, const FPointerEvent&
 
 						// Snap to grid
 						const float SnapSize = GetSnapGridSize();
-						AnchorNodeNewPos.X = SnapSize * FMath::RoundToFloat(AnchorNodeNewPos.X/SnapSize);
-						AnchorNodeNewPos.Y = SnapSize * FMath::RoundToFloat(AnchorNodeNewPos.Y/SnapSize);
+						AnchorNodeNewPos.X = SnapSize * FMath::RoundToFloat( AnchorNodeNewPos.X / SnapSize );
+						AnchorNodeNewPos.Y = SnapSize * FMath::RoundToFloat( AnchorNodeNewPos.Y / SnapSize );
 
 						// Dragging an unselected node automatically selects it.
 						SelectionManager.StartDraggingNode(NodeBeingDragged->GetObjectBeingDisplayed(), MouseEvent);
@@ -1145,6 +1145,11 @@ void SNodePanel::RestoreViewSettings(const FVector2D& InViewOffset, float InZoom
 	OldZoomAmount = GetZoomAmount();
 }
 
+float SNodePanel::GetSnapGridSize()
+{
+	return GetDefault<UEditorStyleSettings>()->GridSnapSize;
+}
+
 inline float FancyMod(float Value, float Size)
 {
 	return ((Value >= 0) ? 0.0f : Size) + FMath::Fmod(Value, Size);
diff --git a/Engine/Source/Editor/GraphEditor/Public/SGraphPin.h b/Engine/Source/Editor/GraphEditor/Public/SGraphPin.h
--- a/Engine/Source/Editor/GraphEditor/Public/SGraphPin.h
+++ b/Engine/Source/Editor/GraphEditor/Public/SGraphPin.h
@@ -118,6 +118,12 @@ public:
 	/** @return whether this pin is an array value */
 	bool IsArray() const;
 
+	/** @return whether this pin is a set value */
+	bool IsSet() const;
+
+	/** @return whether this pin is a map value */
+	bool IsMap() const;
+
 	/** @return whether this pin is passed by ref */
 	bool IsByRef() const;
 
@@ -174,9 +180,12 @@ protected:
 	/** Get the widget we should put into the 'default value' space, shown when nothing connected */
 	virtual TSharedRef<SWidget>	GetDefaultValueWidget();
 
-	/** @return The brush with which to pain this graph pin's incoming/outgoing bullet point */
+	/** @return The brush with which to paint this graph pin's incoming/outgoing bullet point */
 	virtual const FSlateBrush* GetPinIcon() const;
 
+	/** @return the brush which is to paint the 'secondary icon' for this pin, (e.g. value type for Map pins */
+	const FSlateBrush* GetSecondaryPinIcon() const;
+
 	const FSlateBrush* GetPinBorder() const;
 	
 	/** @return The status brush to show after/before the name (on an input/output) */
@@ -191,6 +200,9 @@ protected:
 	/** @return The color that we should use to draw this pin */
 	virtual FSlateColor GetPinColor() const;
 
+	/** @return The secondary color that we should use to draw this pin (e.g. value color for Map pins) */
+	FSlateColor GetSecondaryPinColor() const;
+
 	/** @return The color that we should use to draw this pin's text */
 	virtual FSlateColor GetPinTextColor() const;
 
@@ -245,23 +257,18 @@ protected:
 	bool bUsePinColorForText;
 
 	//@TODO: Want to cache these once for all SGraphPins, but still handle slate style updates
-	mutable const FSlateBrush* CachedImg_ArrayPin_ConnectedHovered;
-	mutable const FSlateBrush* CachedImg_ArrayPin_Connected;
-	mutable const FSlateBrush* CachedImg_ArrayPin_DisconnectedHovered;
-	mutable const FSlateBrush* CachedImg_ArrayPin_Disconnected;
-	mutable const FSlateBrush* CachedImg_RefPin_ConnectedHovered;
-	mutable const FSlateBrush* CachedImg_RefPin_Connected;
-	mutable const FSlateBrush* CachedImg_RefPin_DisconnectedHovered;
-	mutable const FSlateBrush* CachedImg_RefPin_Disconnected;
-	mutable const FSlateBrush* CachedImg_Pin_ConnectedHovered;
-	mutable const FSlateBrush* CachedImg_Pin_Connected;
-	mutable const FSlateBrush* CachedImg_Pin_DisconnectedHovered;
-	mutable const FSlateBrush* CachedImg_Pin_Disconnected;
-	mutable const FSlateBrush* CachedImg_DelegatePin_ConnectedHovered;
-	mutable const FSlateBrush* CachedImg_DelegatePin_Connected;
-	mutable const FSlateBrush* CachedImg_DelegatePin_DisconnectedHovered;
-	mutable const FSlateBrush* CachedImg_DelegatePin_Disconnected;
-
-	mutable const FSlateBrush* CachedImg_Pin_Background;
-	mutable const FSlateBrush* CachedImg_Pin_BackgroundHovered;
+	const FSlateBrush* CachedImg_ArrayPin_Connected;
+	const FSlateBrush* CachedImg_ArrayPin_Disconnected;
+	const FSlateBrush* CachedImg_RefPin_Connected;
+	const FSlateBrush* CachedImg_RefPin_Disconnected;
+	const FSlateBrush* CachedImg_Pin_Connected;
+	const FSlateBrush* CachedImg_Pin_Disconnected;
+	const FSlateBrush* CachedImg_DelegatePin_Connected;
+	const FSlateBrush* CachedImg_DelegatePin_Disconnected;
+	const FSlateBrush* CachedImg_SetPin;
+	const FSlateBrush* CachedImg_MapPinKey;
+	const FSlateBrush* CachedImg_MapPinValue;
+
+	const FSlateBrush* CachedImg_Pin_Background;
+	const FSlateBrush* CachedImg_Pin_BackgroundHovered;
 };
diff --git a/Engine/Source/Editor/GraphEditor/Public/SNodePanel.h b/Engine/Source/Editor/GraphEditor/Public/SNodePanel.h
--- a/Engine/Source/Editor/GraphEditor/Public/SNodePanel.h
+++ b/Engine/Source/Editor/GraphEditor/Public/SNodePanel.h
@@ -650,10 +650,7 @@ public:
 	void RestoreViewSettings(const FVector2D& InViewOffset, float InZoomAmount);
 
 	/** Get the grid snap size */
-	static float GetSnapGridSize()
-	{
-		return 16.f;
-	}
+	static float GetSnapGridSize();
 
 	/** 
 	 * Zooms out to fit either all nodes or only the selected ones.
diff --git a/Engine/Source/Editor/Kismet/Private/BlueprintDetailsCustomization.cpp b/Engine/Source/Editor/Kismet/Private/BlueprintDetailsCustomization.cpp
--- a/Engine/Source/Editor/Kismet/Private/BlueprintDetailsCustomization.cpp
+++ b/Engine/Source/Editor/Kismet/Private/BlueprintDetailsCustomization.cpp
@@ -4843,7 +4843,7 @@ void FBlueprintGlobalOptionsDetails::CustomizeDetails(IDetailLayoutBuilder& Deta
 			UProperty* Property = *PropertyIt;
 			FString Category = Property->GetMetaData(TEXT("Category"));
 
-			if (Category != TEXT("BlueprintOptions") && Category != TEXT("ClassOptions") )
+			if (Category != TEXT("BlueprintOptions") && Category != TEXT("ClassOptions") && Category != TEXT("Experimental"))
 			{
 				DetailLayout.HideProperty(DetailLayout.GetProperty(Property->GetFName()));
 			}
diff --git a/Engine/Source/Editor/Kismet/Private/Profiler/EventExecution.cpp b/Engine/Source/Editor/Kismet/Private/Profiler/EventExecution.cpp
--- a/Engine/Source/Editor/Kismet/Private/Profiler/EventExecution.cpp
+++ b/Engine/Source/Editor/Kismet/Private/Profiler/EventExecution.cpp
@@ -1064,7 +1064,10 @@ void FScriptExecutionBlueprint::RefreshStats(const FTracePath& TracePath)
 	{
 		// This crawls through and updates all instance stats and pools the results into the blueprint node stats
 		// as an overall blueprint performance representation.
-		BlueprintEventIter->RefreshStats(TracePath);
+		if (BlueprintEventIter->HasFlags(EScriptExecutionNodeFlags::RequiresRefresh))
+		{
+			BlueprintEventIter->RefreshStats(TracePath);
+		}
 	}
 	// Update instance stats
 	for (auto InstanceIter : Instances)
@@ -1094,6 +1097,14 @@ void FScriptExecutionBlueprint::RefreshStats(const FTracePath& TracePath)
 	}
 	// Update the node stats
 	UpdatePerfDataForNode();
+	// Mark all the stats as updated
+	for (auto BlueprintEventIter : ChildNodes)
+	{
+		if (BlueprintEventIter->HasFlags(EScriptExecutionNodeFlags::RequiresRefresh))
+		{
+			BlueprintEventIter->RemoveFlags(EScriptExecutionNodeFlags::RequiresRefresh);
+		}
+	}
 }
 
 void FScriptExecutionBlueprint::GetAllExecNodes(TMap<FName, TSharedPtr<FScriptExecutionNode>>& ExecNodesOut)
@@ -1238,12 +1249,15 @@ void FScriptExecutionBlueprint::UpdateHeatLevels()
 		FTracePath RootTracePath;
 		for (auto EventNode : ChildNodes)
 		{
-			// Grab the instance perf data
-			TSharedPtr<FScriptPerfData> InstancePerfData = EventNode->GetOrAddPerfDataByInstanceAndTracePath(CurrentInstanceName, RootTracePath);
-			// Create the heat update data
-			FScriptExecutionHottestPathParams HotPathParams(CurrentInstanceName, InstancePerfData->GetAverageTiming(), InstancePerfData->GetSampleCount(), HeatLevelMetrics);
-			// Walk through the event linked nodes and update heat display stats.
-			EventNode->UpdateHeatDisplayStats(HotPathParams);
+			if (EventNode->HasFlags(EScriptExecutionNodeFlags::RequiresRefresh))
+			{
+				// Grab the instance perf data
+				TSharedPtr<FScriptPerfData> InstancePerfData = EventNode->GetOrAddPerfDataByInstanceAndTracePath(CurrentInstanceName, RootTracePath);
+				// Create the heat update data
+				FScriptExecutionHottestPathParams HotPathParams(CurrentInstanceName, InstancePerfData->GetAverageTiming(), InstancePerfData->GetSampleCount(), HeatLevelMetrics);
+				// Walk through the event linked nodes and update heat display stats.
+				EventNode->UpdateHeatDisplayStats(HotPathParams);
+			}
 		}
 	}
 }
diff --git a/Engine/Source/Editor/Kismet/Private/SBlueprintEditorSelectedDebugObjectWidget.cpp b/Engine/Source/Editor/Kismet/Private/SBlueprintEditorSelectedDebugObjectWidget.cpp
--- a/Engine/Source/Editor/Kismet/Private/SBlueprintEditorSelectedDebugObjectWidget.cpp
+++ b/Engine/Source/Editor/Kismet/Private/SBlueprintEditorSelectedDebugObjectWidget.cpp
@@ -252,7 +252,6 @@ void SBlueprintEditorSelectedDebugObjectWidget::GenerateDebugWorldNames(bool bRe
 			continue;
 		}
 
-		DebugWorlds.Add(TestWorld);
 		ENetMode NetMode = TestWorld->GetNetMode();
 
 		FString WorldName;
@@ -281,6 +280,11 @@ void SBlueprintEditorSelectedDebugObjectWidget::GenerateDebugWorldNames(bool bRe
 
 		if (!WorldName.IsEmpty())
 		{
+			// DebugWorlds & DebugWorldNames need to be the same size (we expect
+			// an index in one to correspond to the other) - DebugWorldNames is
+			// what populates the dropdown, so it is the authority (if there's 
+			// no name to present, they can't select from DebugWorlds)
+			DebugWorlds.Add(TestWorld);
 			DebugWorldNames.Add( MakeShareable(new FString(WorldName)) );
 		}
 	}
diff --git a/Engine/Source/Editor/Kismet/Private/SCSEditorViewportClient.cpp b/Engine/Source/Editor/Kismet/Private/SCSEditorViewportClient.cpp
--- a/Engine/Source/Editor/Kismet/Private/SCSEditorViewportClient.cpp
+++ b/Engine/Source/Editor/Kismet/Private/SCSEditorViewportClient.cpp
@@ -921,24 +921,13 @@ void FSCSEditorViewportClient::BeginTransaction(const FText& Description)
 						}
 					}
 
-					// Modify both the component template and the instance in the preview actor (provided there is one)
+					// Modify template, any instances will be reconstructed as part of PostUndo:
 					UActorComponent* ComponentTemplate = Node->GetEditableComponentTemplate(PreviewBlueprint);
 					if (ComponentTemplate != nullptr)
 					{
 						ComponentTemplate->SetFlags(RF_Transactional);
 						ComponentTemplate->Modify();
 					}
-
-					AActor* PreviewActor = GetPreviewActor();
-					if (PreviewActor)
-					{
-						UActorComponent* ComponentPreviewInstance = Node->FindComponentInstanceInActor(PreviewActor);
-						if (ComponentPreviewInstance != nullptr)
-						{
-							ComponentPreviewInstance->SetFlags(RF_Transactional);
-							ComponentPreviewInstance->Modify();
-						}
-					}
 				}
 			}
 		}
diff --git a/Engine/Source/Editor/Kismet/Private/SSCSEditor.cpp b/Engine/Source/Editor/Kismet/Private/SSCSEditor.cpp
--- a/Engine/Source/Editor/Kismet/Private/SSCSEditor.cpp
+++ b/Engine/Source/Editor/Kismet/Private/SSCSEditor.cpp
@@ -4725,14 +4725,31 @@ UActorComponent* SSCSEditor::AddNewNode(USCS_Node* NewNode, UObject* Asset, bool
 	UBlueprint* Blueprint = GetBlueprint();
 	check(Blueprint != nullptr && Blueprint->SimpleConstructionScript != nullptr);
 
-	// Reset the scene root node if it's set to the default one that's managed by the SCS
-	//if(SceneRootNodePtr.IsValid() && SceneRootNodePtr->GetSCSNode() == Blueprint->SimpleConstructionScript->GetDefaultSceneRootNode())
-	//{
-	//	SceneRootNodePtr.Reset();
-	//}
+	bool AttachToSceneRootNode = true;
+	if (USceneComponent* NewSceneComponent = Cast<USceneComponent>(NewNode->ComponentTemplate))
+	{
+		// get currently selected component
+		TArray<FSCSEditorTreeNodePtrType> SelectedTreeNodes;
+		if (SCSTreeWidget.IsValid() && SCSTreeWidget->GetSelectedItems(SelectedTreeNodes) > 0)
+		{
+			FSCSEditorTreeNodePtrType FirstTreeNode = SelectedTreeNodes[0];
+			if (FirstTreeNode.IsValid() && FirstTreeNode->GetComponentTemplate())
+			{
+				USceneComponent* CastFirstTreeNode = Cast<USceneComponent>(FirstTreeNode->GetComponentTemplate());
+				if (CastFirstTreeNode && NewSceneComponent->CanAttachAsChild(CastFirstTreeNode, NAME_None))
+				{
+					NewNodePtr = AddTreeNode(NewNode, FirstTreeNode, false);
+					AttachToSceneRootNode = false;
+				}
+			}
+		}
+	}
 
-	// Add the new node to the editor tree
-	NewNodePtr = AddTreeNode(NewNode, SceneRootNodePtr, false);
+	if (AttachToSceneRootNode)
+	{
+		// Add the new node to the editor tree
+		NewNodePtr = AddTreeNode(NewNode, SceneRootNodePtr, false);
+	}
 
 	// Potentially adjust variable names for any child blueprints
 	const FName VariableName = NewNode->GetVariableName();
diff --git a/Engine/Source/Editor/Kismet/Public/Profiler/EventExecution.h b/Engine/Source/Editor/Kismet/Public/Profiler/EventExecution.h
--- a/Engine/Source/Editor/Kismet/Public/Profiler/EventExecution.h
+++ b/Engine/Source/Editor/Kismet/Public/Profiler/EventExecution.h
@@ -36,6 +36,7 @@ namespace EScriptExecutionNodeFlags
 		InvalidTrace				= 0x00400000,	// Indicates that node doesn't contain a valid script trace.
 		RuntimeEvent				= 0x01000000,	// Event that is considered part of the runtime cost.
 		ConstructionEvent			= 0x02000000,	// Event that is considered part of the construction cost.
+		RequiresRefresh				= 0x04000000,	// This event has updated stats and requires a refresh.
 		// Groups
 		CallSite					= FunctionCall|ParentFunctionCall|MacroCall|TunnelInstance,
 		BranchNode					= ConditionalBranch|SequentialBranch,
diff --git a/Engine/Source/Editor/KismetCompiler/Private/KismetCompiler.cpp b/Engine/Source/Editor/KismetCompiler/Private/KismetCompiler.cpp
--- a/Engine/Source/Editor/KismetCompiler/Private/KismetCompiler.cpp
+++ b/Engine/Source/Editor/KismetCompiler/Private/KismetCompiler.cpp
@@ -2195,6 +2195,15 @@ void FKismetCompilerContext::CreatePinEventNodeForTimelineFunction(UK2Node_Timel
 	{
 		MovePinLinksToIntermediate(*UpdatePin, *UpdateOutput);
 	}
+
+	if (CompileOptions.bAddInstrumentation)
+	{
+		// Register the timeline event nodes to any composite instances the timeline is located in.
+		if (UEdGraphNode* TunnelInstance = MessageLog.GetIntermediateTunnelInstance(TimelineNode))
+		{
+			MessageLog.RegisterIntermediateTunnelNode(TimelineEventNode, TunnelInstance);
+		}
+	}
 }
 
 UK2Node_CallFunction* FKismetCompilerContext::CreateCallTimelineFunction(UK2Node_Timeline* TimelineNode, UEdGraph* SourceGraph, FName FunctionName, UEdGraphPin* TimelineVarPin, UEdGraphPin* TimelineFunctionPin)
@@ -2694,11 +2703,14 @@ void FKismetCompilerContext::ExpansionStep(UEdGraph* Graph, bool bAllowUbergraph
 					// Keep track of new expansion nodes
 					if (ExpansionNodeIdx < Graph->Nodes.Num())
 					{
+						bool bImpureExpansionNodeAdded = false;
 						TArray<UEdGraphNode*> ExpansionNodes;
 						for (; ExpansionNodeIdx < Graph->Nodes.Num(); ++ExpansionNodeIdx)
 						{
 							ExpansionNodes.Add(Graph->Nodes[ExpansionNodeIdx]);
 							MessageLog.RegisterExpansionNode(Graph->Nodes[ExpansionNodeIdx]);
+							// We don't want to create boundary nodes around a set of pure only expansion nodes.
+							bImpureExpansionNodeAdded |= !IsNodePure(Graph->Nodes[ExpansionNodeIdx]);
 							// Add the relevent pins to the trace suppression filter
 							for (auto ExpansionPin : Graph->Nodes[ExpansionNodeIdx]->Pins)
 							{
@@ -2709,8 +2721,11 @@ void FKismetCompilerContext::ExpansionStep(UEdGraph* Graph, bool bAllowUbergraph
 							}
 						}
 						// Create boundaries around expsnion nodes so we can simulate pin traces correctly
-						UK2Node_TunnelBoundary::CreateBoundariesForExpansionNodes(Node, ExpansionNodes, SourceNodeLinks, MessageLog);
-						ExpansionNodeIdx = Graph->Nodes.Num();
+						if (bImpureExpansionNodeAdded)
+						{
+							UK2Node_TunnelBoundary::CreateBoundariesForExpansionNodes(Node, ExpansionNodes, SourceNodeLinks, MessageLog);
+							ExpansionNodeIdx = Graph->Nodes.Num();
+						}
 					}
 				}
 			}
@@ -3852,9 +3867,9 @@ void FKismetCompilerContext::Compile()
 			GConfig->GetBool(TEXT("Kismet"), TEXT("CompileDisplaysBinaryBackend"), /*out*/ bDisplayBytecode, GEngineIni);
 		}
 
-		const bool bGenerateStubsOnly = !bIsFullCompile || (0 != MessageLog.NumErrors);
 		// Always run the VM backend, it's needed for more than just debug printing
 		{
+			const bool bGenerateStubsOnly = !bIsFullCompile || (0 != MessageLog.NumErrors);
 			BP_SCOPED_COMPILER_EVENT_STAT(EKismetCompilerStats_CodeGenerationTime);
 			Backend_VM.GenerateCodeFromClass(NewClass, FunctionList, bGenerateStubsOnly);
 			if (!bGenerateStubsOnly)
@@ -3863,7 +3878,8 @@ void FKismetCompilerContext::Compile()
 			}
 		}
 
-		if (!bGenerateStubsOnly)
+		// Fill ScriptObjectReferences arrays in functions
+		if (bIsFullCompile && (0 == MessageLog.NumErrors)) // Backend_VM can generate errors, so bGenerateStubsOnly cannot be reused
 		{
 			for (FKismetFunctionContext& FunctionContext : FunctionList)
 			{
diff --git a/Engine/Source/Editor/KismetCompiler/Private/KismetCompilerVMBackend.cpp b/Engine/Source/Editor/KismetCompiler/Private/KismetCompilerVMBackend.cpp
--- a/Engine/Source/Editor/KismetCompiler/Private/KismetCompilerVMBackend.cpp
+++ b/Engine/Source/Editor/KismetCompiler/Private/KismetCompilerVMBackend.cpp
@@ -16,6 +16,7 @@
 #include "Engine/UserDefinedStruct.h"
 
 #include "TextPackageNamespaceUtil.h"
+#include "BlueprintEditorUtils.h"
 
 #define LOCTEXT_NAMESPACE "KismetCompilerVMBackend"
 //////////////////////////////////////////////////////////////////////////
@@ -297,9 +298,7 @@ protected:
 				}
 			}
 
-			// Function contexts must always be objects, so if we have a literal, give it the default object property so the compiler knows how to handle it
-			UProperty* CoerceProperty = Term->bIsLiteral ? ((UProperty*)(GetDefault<UObjectProperty>())) : NULL;
-			ScriptBuilder.EmitTerm(Term, CoerceProperty);
+			ScriptBuilder.EmitTerm(Term);
 
 			// Skip offset if the expression evaluates to null (counting from later on)
 			FSkipOffsetEmitter Skipper(ScriptBuilder.Writer.ScriptBuffer);
@@ -497,6 +496,7 @@ public:
 			return Type && (Type->PinCategory == UEdGraphSchema_K2::PC_Asset);
 		}
 
+		// Will handle Class properties as well
 		static bool IsObject(const FEdGraphPinType* Type, const UProperty* Property)
 		{
 			if (Property)
@@ -506,6 +506,15 @@ public:
 			return Type && (Type->PinCategory == UEdGraphSchema_K2::PC_Object);
 		}
 
+		static bool IsClass(const FEdGraphPinType* Type, const UProperty* Property)
+		{
+			if (Property)
+			{
+				return Property->IsA<UClassProperty>();
+			}
+			return Type && (Type->PinCategory == UEdGraphSchema_K2::PC_Class);
+		}
+
 		static bool IsInterface(const FEdGraphPinType* Type, const UProperty* Property)
 		{
 			if (Property)
@@ -522,6 +531,43 @@ public:
 		{
 			check(!Term->Type.bIsArray || CoerceProperty);
 
+			// Additional Validation, since we cannot trust custom k2nodes
+			const bool bSecialCaseSelf = (Term->Type.PinSubCategory == Schema->PN_Self);
+			if (CoerceProperty && ensure(Schema) && ensure(CurrentCompilerContext) && !bSecialCaseSelf)
+			{
+				FEdGraphPinType TrueType;
+				const bool bValidProperty = Schema->ConvertPropertyToPinType(CoerceProperty, TrueType);
+
+				auto AreTypesBinaryCompatible = [](const FEdGraphPinType& TypeA, const FEdGraphPinType& TypeB) -> bool
+				{
+					if (TypeA.PinCategory != TypeB.PinCategory)
+					{
+						return false;
+					}
+					if ((TypeA.bIsMap != TypeB.bIsMap)
+						|| (TypeA.bIsSet != TypeB.bIsSet)
+						|| (TypeA.bIsArray != TypeB.bIsArray)
+						|| (TypeA.bIsWeakPointer != TypeB.bIsWeakPointer))
+					{
+						return false;
+					}
+					if (TypeA.PinCategory == UEdGraphSchema_K2::PC_Struct)
+					{
+						if (TypeA.PinSubCategoryObject != TypeB.PinSubCategoryObject)
+						{
+							return false;
+						}
+					}
+					return true;
+				};
+
+				if (!bValidProperty || !AreTypesBinaryCompatible(Term->Type, TrueType))
+				{
+					const FString ErrorMessage = FString::Printf(TEXT("ICE: The type of property %s doesn't match a term. @@"), *CoerceProperty->GetPathName());
+					CurrentCompilerContext->MessageLog.Error(*ErrorMessage, Term->SourcePin);
+				}
+			}
+
 			if (FLiteralTypeHelper::IsString(&Term->Type, CoerceProperty))
 			{
 				EmitStringLiteral(Term->Name);
@@ -656,10 +702,6 @@ public:
 				Writer << EX_NameConst;
 				Writer << LiteralName;
 			}
-			//else if (UClassProperty* ClassProperty = Cast<UClassProperty>(CoerceProperty))
-			//{
-			//	ensureMsgf(false, TEXT("Class property literals are not supported yet!"));
-			//}
 			else if (FLiteralTypeHelper::IsStruct(&Term->Type, CoerceProperty))
 			{
 				UStructProperty* StructProperty = Cast<UStructProperty>(CoerceProperty);
@@ -721,12 +763,14 @@ public:
 					Writer << Struct;
 					Writer << StructSize;
 
+					checkSlow(Schema);
 					for( UProperty* Prop = Struct->PropertyLink; Prop; Prop = Prop->PropertyLinkNext )
 					{
 						for (int32 ArrayIter = 0; ArrayIter < Prop->ArrayDim; ++ArrayIter)
 						{
 							// Create a new term for each property, and serialize it out
 							FBPTerminal NewTerm;
+							Schema->ConvertPropertyToPinType(Prop, NewTerm.Type);
 							NewTerm.bIsLiteral = true;
 							NewTerm.Source = Term->Source;
 							Prop->ExportText_InContainer(ArrayIter, NewTerm.Name, StructData, StructData, NULL, PPF_None);
@@ -763,6 +807,7 @@ public:
 				for (int32 ElemIdx = 0; ElemIdx < ElementNum; ++ElemIdx)
 				{
 					FBPTerminal NewTerm;
+					Schema->ConvertPropertyToPinType(InnerProp, NewTerm.Type);
 					NewTerm.bIsLiteral = true;
 					NewTerm.Source = Term->Source;
 					uint8* RawElemData = ScriptArrayHelper.GetRawPtr(ElemIdx);
@@ -800,7 +845,7 @@ public:
 				FAssetPtr AssetPtr(Term->ObjectLiteral);
 				EmitStringLiteral(AssetPtr.GetUniqueID().ToString());
 			}
-			else if (FLiteralTypeHelper::IsObject(&Term->Type, CoerceProperty))
+			else if (FLiteralTypeHelper::IsObject(&Term->Type, CoerceProperty) || FLiteralTypeHelper::IsClass(&Term->Type, CoerceProperty))
 			{
 				// Note: This case handles both UObjectProperty and UClassProperty
 				if (Term->Type.PinSubCategory == Schema->PN_Self)
@@ -832,6 +877,10 @@ public:
 					ensureMsgf(false, TEXT("It is not possible to express this interface property as a literal value! (%s)"), *CoerceProperty->GetFullName());
 				}
 			}
+			else if (!CoerceProperty && Term->Type.PinCategory.IsEmpty() && (Term->Type.PinSubCategory == Schema->PN_Self))
+			{
+				Writer << EX_Self;
+			}
 			// else if (CoerceProperty->IsA(UMulticastDelegateProperty::StaticClass()))
 			// Cannot assign a literal to a multicast delegate; it should be added instead of assigned
 			else
@@ -881,6 +930,7 @@ public:
 		Writer << LatentInfoStruct;
 		Writer << StructSize;
 
+		checkSlow(Schema);
 		for (UProperty* Prop = LatentInfoStruct->PropertyLink; Prop; Prop = Prop->PropertyLinkNext)
 		{
 			if (TargetLabel && Prop->GetBoolMetaData(FBlueprintMetadata::MD_NeedsLatentFixup))
@@ -901,6 +951,7 @@ public:
 			{
 				// Create a new term for each property, and serialize it out
 				FBPTerminal NewTerm;
+				Schema->ConvertPropertyToPinType(Prop, NewTerm.Type);
 				NewTerm.bIsLiteral = true;
 				Prop->ExportText_InContainer(0, NewTerm.Name, StructData, StructData, NULL, PPF_None);
 
@@ -1612,7 +1663,8 @@ public:
 			{
 				// If this is a debug site for an expanded macro instruction, there should also be a macro source node associated with it
 				UEdGraphNode* MacroSourceNode = CompilerContext.MessageLog.GetSourceTunnelNode(SourceNode);
-				if (MacroSourceNode == TrueSourceNode)
+				// We need to ensure that macro/composite instances also record the tunnels present at any script location ( including themselves )
+				if (MacroSourceNode == TrueSourceNode && !FBlueprintEditorUtils::IsTunnelInstanceNode(MacroSourceNode))
 				{
 					// The function above will return the given node if not found in the map. In that case there is no associated source macro node, so we clear it.
 					MacroSourceNode = NULL;
diff --git a/Engine/Source/Editor/KismetWidgets/Private/SPinTypeSelector.cpp b/Engine/Source/Editor/KismetWidgets/Private/SPinTypeSelector.cpp
--- a/Engine/Source/Editor/KismetWidgets/Private/SPinTypeSelector.cpp
+++ b/Engine/Source/Editor/KismetWidgets/Private/SPinTypeSelector.cpp
@@ -148,6 +148,14 @@ TSharedRef<SWidget> SPinTypeSelector::ConstructPinTypeImage(const FSlateBrush* P
 		.ColorAndOpacity(PrimaryColor);
 }
 
+TSharedRef<SWidget> SPinTypeSelector::ConstructPinTypeImage(TAttribute<const FSlateBrush*> PrimaryIcon, TAttribute<FSlateColor> PrimaryColor, TAttribute<const FSlateBrush*> SecondaryIcon, TAttribute<FSlateColor> SecondaryColor )
+{
+	return
+		SNew(SDoubleImage, SecondaryIcon, SecondaryColor)
+		.Image(PrimaryIcon)
+		.ColorAndOpacity(PrimaryColor);
+}
+
 void SPinTypeSelector::Construct(const FArguments& InArgs, FGetPinTypeTree GetPinTypeTreeFunc)
 {
 	// SComboBox is a bit restrictive:
@@ -775,13 +783,19 @@ void SPinTypeSelector::OnSelectPinType(FPinTypeTreeItem InItem, FString InPinCat
 
 	if ((NewTargetPinType.bIsMap || NewTargetPinType.bIsSet) && !FBlueprintEditorUtils::HasGetTypeHash(NewTargetPinType))
 	{
+		FEdGraphPinType HashedType = NewTargetPinType;
+		// clear the container-ness for messaging, we want to explain that the contained type is not hashable,
+		// not message about the container type (e.g. "Container type cleared because 'bool' does not have a GetTypeHash..."
+		// instead of "Container Type cleared because 'map of bool to float'..."). We also need to clear this because
+		// the type cannot be a container:
+		NewTargetPinType.bIsMap = false;
+		NewTargetPinType.bIsSet = false;
+
 		// inform user via toast why the type change was exceptional and clear IsMap/IsSetness because this type cannot be hashed:
-		const FText NotificationText = FText::Format(LOCTEXT("TypeCannotBeHashed", "Container type cleared because '{0}' does not have a GetTypeHash function. Maps and Sets require a hash function to insert and find elements"), GetTypeDescription());
+		const FText NotificationText = FText::Format(LOCTEXT("TypeCannotBeHashed", "Container type cleared because '{0}' does not have a GetTypeHash function. Maps and Sets require a hash function to insert and find elements"), UEdGraphSchema_K2::TypeToText(NewTargetPinType));
 		FNotificationInfo Info(NotificationText);
 		Info.ExpireDuration = 8.0f;
 		FSlateNotificationManager::Get().AddNotification(Info);
-		NewTargetPinType.bIsMap = false;
-		NewTargetPinType.bIsSet = false;
 	}
 
 	OnTypeChanged.ExecuteIfBound(NewTargetPinType);
diff --git a/Engine/Source/Editor/KismetWidgets/Public/SPinTypeSelector.h b/Engine/Source/Editor/KismetWidgets/Public/SPinTypeSelector.h
--- a/Engine/Source/Editor/KismetWidgets/Public/SPinTypeSelector.h
+++ b/Engine/Source/Editor/KismetWidgets/Public/SPinTypeSelector.h
@@ -23,6 +23,7 @@ class KISMETWIDGETS_API SPinTypeSelector : public SCompoundWidget
 {
 public:
 	static TSharedRef<SWidget> ConstructPinTypeImage(const FSlateBrush* PrimaryIcon, const FSlateColor& PrimaryColor, const FSlateBrush* SecondaryIcon, const FSlateColor& SecondaryColor, TSharedPtr<SToolTip> InToolTip);
+	static TSharedRef<SWidget> ConstructPinTypeImage(TAttribute<const FSlateBrush*> PrimaryIcon, TAttribute<FSlateColor> PrimaryColor, TAttribute<const FSlateBrush*> SecondaryIcon, TAttribute<FSlateColor> SecondaryColor );
 
 	SLATE_BEGIN_ARGS( SPinTypeSelector )
 		: _TargetPinType()
diff --git a/Engine/Source/Editor/UnrealEd/Classes/Settings/ProjectPackagingSettings.h b/Engine/Source/Editor/UnrealEd/Classes/Settings/ProjectPackagingSettings.h
--- a/Engine/Source/Editor/UnrealEd/Classes/Settings/ProjectPackagingSettings.h
+++ b/Engine/Source/Editor/UnrealEd/Classes/Settings/ProjectPackagingSettings.h
@@ -86,6 +86,10 @@ public:
 	UPROPERTY(config, EditAnywhere, Category=Experimental)
 	bool bNativizeBlueprintAssets;
 
+	/** WHen nativization is enabled, only Blueprints with "Nativize" set true will be converted to C++. */
+	UPROPERTY(config, EditAnywhere, Category = Experimental)
+	bool bNativizeOnlySelectedBlueprints;
+
 	/** If enabled, all content will be put into a single .pak file instead of many individual files (default = enabled). */
 	UPROPERTY(config, EditAnywhere, Category=Packaging)
 	bool UsePakFile;
diff --git a/Engine/Source/Editor/UnrealEd/Private/EdGraphUtilities.cpp b/Engine/Source/Editor/UnrealEd/Private/EdGraphUtilities.cpp
--- a/Engine/Source/Editor/UnrealEd/Private/EdGraphUtilities.cpp
+++ b/Engine/Source/Editor/UnrealEd/Private/EdGraphUtilities.cpp
@@ -204,6 +204,7 @@ UEdGraph* FEdGraphUtilities::CloneGraph(UEdGraph* InSource, UObject* NewOuter, F
 	if (bCloningForCompile || (NewOuter == NULL))
 	{
 		Parameters.ApplyFlags |= RF_Transient;
+		Parameters.FlagMask &= ~RF_Transactional;
 	}
 
 	UEdGraph* ClonedGraph = CastChecked<UEdGraph>(StaticDuplicateObjectEx(Parameters));
@@ -483,6 +484,37 @@ bool FEdGraphUtilities::IsSetParam(const UFunction* Function, const FString& Par
 	return false;
 }
 
+bool FEdGraphUtilities::IsMapParam(const UFunction* Function, const FString& ParameterName)
+{
+	if (Function == nullptr)
+	{
+		return false;
+	}
+
+	const FString& MapParamMetaData = Function->GetMetaData(FBlueprintMetadata::MD_MapParam);
+	const FString& MapValueParamMetaData = Function->GetMetaData(FBlueprintMetadata::MD_MapValueParam);
+	const FString& MapKeyParamMetaData = Function->GetMetaData(FBlueprintMetadata::MD_MapKeyParam);
+	if (MapParamMetaData.IsEmpty() && MapValueParamMetaData.IsEmpty() && MapKeyParamMetaData.IsEmpty() )
+	{
+		return false;
+	}
+
+	const auto PipeSeparatedStringContains = [ParameterName](const FString& List)
+	{
+		TArray<FString> GroupEntries;
+		List.ParseIntoArray(GroupEntries, TEXT("|"), true);
+		if (GroupEntries.Contains(ParameterName))
+		{
+			return true;
+		}
+		return false;
+	};
+
+	return PipeSeparatedStringContains(MapParamMetaData)
+		|| PipeSeparatedStringContains(MapValueParamMetaData) 
+		|| PipeSeparatedStringContains(MapKeyParamMetaData);
+}
+
 void FEdGraphUtilities::RegisterVisualNodeFactory(TSharedPtr<FGraphPanelNodeFactory> NewFactory)
 {
 	VisualNodeFactories.Add(NewFactory);
diff --git a/Engine/Source/Editor/UnrealEd/Private/Kismet2/BlueprintEditorUtils.cpp b/Engine/Source/Editor/UnrealEd/Private/Kismet2/BlueprintEditorUtils.cpp
--- a/Engine/Source/Editor/UnrealEd/Private/Kismet2/BlueprintEditorUtils.cpp
+++ b/Engine/Source/Editor/UnrealEd/Private/Kismet2/BlueprintEditorUtils.cpp
@@ -1158,142 +1158,208 @@ struct FRegenerationHelper
 	}
 };
 
-struct FEditoronlyBlueprintHelper
+bool FLegacyEditorOnlyBlueprintUtils::DoPinsMatch(const FEdGraphPinType& Input, const FEdGraphPinType& Output)
 {
-	static bool IsUnwantedType(const FEdGraphPinType& Type)
+	bool bHasBlueprintPin = false;
+	auto TestCompatibility = [&bHasBlueprintPin](const FEdGraphPinType& A, const FEdGraphPinType& B)
 	{
-		if (const UClass* BPClass = Cast<const UClass>(Type.PinSubCategoryObject.Get()))
+		if ((B.PinCategory == UEdGraphSchema_K2::PC_Object) && (A.PinCategory == UEdGraphSchema_K2::PC_Class))
 		{
-			if (BPClass->IsChildOf(UBlueprint::StaticClass()))
-			{
-				return true;
-			}
-		}
+			const bool bAIsObjectClass = (UObject::StaticClass() == A.PinSubCategoryObject.Get());
+
+			const UClass* BClass = Cast<UClass>(B.PinSubCategoryObject.Get());
+			const bool bBIsBlueprintObj = BClass && BClass->IsChildOf(UBlueprint::StaticClass());
+			bHasBlueprintPin |= bBIsBlueprintObj;
 
+			return bAIsObjectClass && bBIsBlueprintObj;
+		}
 		return false;
-	}
+	};
 
-	static bool IsUnwantedDefaultObject(const UObject* Obj)
+	const bool bIsMatch = TestCompatibility(Input, Output) || TestCompatibility(Output, Input);
+	if (bHasBlueprintPin)
 	{
-		return Obj && Obj->IsA<UBlueprint>();
+		if (AllowLegacyBlueprintPinMatchesWithClass())
+		{
+			UE_LOG(LogBlueprint, Warning, TEXT("Editor config [EditoronlyBP] settings are DEPRECATED. You currently have old (now prohibited) UBlueprint pins that are being fixed up on load. Please update/resave the asset to keep it from breaking."));
+			return bIsMatch;
+		}
+		else
+		{
+			UE_LOG(LogBlueprint, Warning, TEXT("Editor config [EditoronlyBP] settings are DEPRECATED. You currently use old (now prohibited) UBlueprint pins. Please switch these over to use the Blueprint's class instead."));
+		}
 	}
+	return false;
+}
 
-	static void ChangePinType(FEdGraphPinType& Type)
+bool FLegacyEditorOnlyBlueprintUtils::FixupBlueprint(UBlueprint* Blueprint)
+{
+	if (HasLegacyBlueprintReferences(Blueprint) && FixupLegacyBlueprintReferences())
 	{
 		const UEdGraphSchema_K2* Schema = GetDefault<UEdGraphSchema_K2>();
 
-		Type.PinCategory = Schema->PC_Class;
-		Type.PinSubCategoryObject = UObject::StaticClass();
-	}
+		//any native function call with blueprint parameter
+		TArray<UEdGraph*> Graphs;
+		Blueprint->GetAllGraphs(Graphs);
+		for (UEdGraph* Graph : Graphs)
+		{
+			Schema->BackwardCompatibilityNodeConversion(Graph, false);
+		}
 
-	static bool ShouldBeFixed(const UBlueprint* Blueprint, bool bLogWhy)
-	{
-		check(Blueprint);
-		//any blueprint variable
-		for (const FBPVariableDescription& VarDesc : Blueprint->NewVariables)
+		//any blueprint pin (function/event/macro parameter
+		for (const UEdGraph* Graph : Graphs)
 		{
-			if (IsUnwantedType(VarDesc.VarType))
 			{
-				if (bLogWhy)
+				TArray<UK2Node_EditablePinBase*> EditablePinNodes;
+				Graph->GetNodesOfClass(EditablePinNodes);
+				for (UK2Node_EditablePinBase* Node : EditablePinNodes)
 				{
-					const FString UnwantedType = GetNameSafe(VarDesc.VarType.PinSubCategoryObject.Get());
-					UE_LOG(LogBlueprint, Warning, TEXT("FEditoronlyBlueprintHelper::ShouldBeFixed. [%s] Unwanted type '%s' in variable '%s'"), *Blueprint->GetName(), *UnwantedType, *VarDesc.FriendlyName);
+					HandleEditablePinNode(Node);
+				}
+			}
+
+			{
+				TArray<UK2Node_TemporaryVariable*> TempVariablesNodes;
+				Graph->GetNodesOfClass(TempVariablesNodes);
+				for (UK2Node_TemporaryVariable* Node : TempVariablesNodes)
+				{
+					HandleTemporaryVariableNode(Node);
 				}
-				return true;
 			}
 		}
 
-		//anything on graph
-		TArray<UEdGraph*> Graphs;
-		Blueprint->GetAllGraphs(Graphs);
+		// change variables
+		for (FBPVariableDescription& VarDesc : Blueprint->NewVariables)
+		{
+			if (IsUnwantedType(VarDesc.VarType))
+			{
+				ChangePinType(VarDesc.VarType);
+			}
+		}
+
 		for (const UEdGraph* Graph : Graphs)
 		{
 			TArray<UK2Node*> Nodes;
 			Graph->GetNodesOfClass(Nodes);
-			for (const UK2Node* Node : Nodes)
+			for (UK2Node* Node : Nodes)
+			{
+				HandleDefaultObjects(Node);
+			}
+		}
+
+		return true;
+	}
+	return false;
+}
+
+bool FLegacyEditorOnlyBlueprintUtils::HasLegacyBlueprintReferences(UBlueprint* Blueprint)
+{
+	check(Blueprint);
+	//any blueprint variable
+	for (const FBPVariableDescription& VarDesc : Blueprint->NewVariables)
+	{
+		if (IsUnwantedType(VarDesc.VarType))
+		{
+			const FString UnwantedType = GetNameSafe(VarDesc.VarType.PinSubCategoryObject.Get());
+			UE_LOG(LogBlueprint, Warning, TEXT("Editor config [EditoronlyBP] settings are DEPRECATED. A Blueprint ('%s') has a variable ('%s') that is now a deprecated type ('%s') - use the target's class instead (may just require a resave)."), *Blueprint->GetName(), *VarDesc.FriendlyName, *UnwantedType);
+			return true;
+		}
+	}
+
+	//anything on graph
+	TArray<UEdGraph*> Graphs;
+	Blueprint->GetAllGraphs(Graphs);
+	for (const UEdGraph* Graph : Graphs)
+	{
+		TArray<UK2Node*> Nodes;
+		Graph->GetNodesOfClass(Nodes);
+		for (const UK2Node* Node : Nodes)
+		{
+			for (const UEdGraphPin* Pin : Node->Pins)
 			{
-				for (const UEdGraphPin* Pin : Node->Pins)
+				if (Pin)
 				{
-					if (Pin)
+					const bool bUnwantedType = IsUnwantedType(Pin->PinType);
+					const bool bUnwantedDefaultObject = IsUnwantedDefaultObject(Pin->DefaultObject);
+					if (bUnwantedType || bUnwantedDefaultObject)
 					{
-						const bool bUnwantedType = IsUnwantedType(Pin->PinType);
-						const bool bUnwantedDefaultObject = IsUnwantedDefaultObject(Pin->DefaultObject);
-						if (bUnwantedType || bUnwantedDefaultObject)
+						const FString ReasonPrefix = FString::Printf(TEXT("Editor config [EditoronlyBP] settings are DEPRECATED. A Blueprint (%s) has a"), *Blueprint->GetName());
+						const FString PinName = Pin->GetDisplayName().ToString();
+						const FString PinNodeName = Pin->GetOwningNode()->GetNodeTitle(ENodeTitleType::ListView).ToString();
+						if (bUnwantedType)
 						{
-							if (bLogWhy)
-							{
-								const FString ReasonPrefix = FString::Printf(TEXT("FEditoronlyBlueprintHelper::ShouldBeFixed. [%s]"), *Blueprint->GetName());
-								const FString PinName = Pin->GetDisplayName().ToString();
-								const FString PinNodeName = Pin->GetOwningNode()->GetNodeTitle(ENodeTitleType::ListView).ToString();
-								if (bUnwantedType)
-								{
-									const FString UnwantedType = GetNameSafe(Pin->PinType.PinSubCategoryObject.Get());
-									UE_LOG(LogBlueprint, Warning, TEXT("%s Unwanted type '%s' on pin '%s' on node '%s'"), *ReasonPrefix, *UnwantedType, *PinName, *PinNodeName);
-								}
-								else if (bUnwantedDefaultObject)
-								{
-									const FString UnwantedDefaultObject = GetNameSafe(Pin->DefaultObject);
-									UE_LOG(LogBlueprint, Warning, TEXT("%s Unwanted default object '%s' on pin '%s' on node '%s'"), *ReasonPrefix, *UnwantedDefaultObject, *PinName, *PinNodeName);
-								}
-								else
-								{
-									ensureMsgf(false, TEXT("Can not describe why the blueprint should be fixed."));
-									UE_LOG(LogBlueprint, Warning, TEXT("%s Unknown reason. Pin '%s' on node '%s'"), *ReasonPrefix, *PinName, *PinNodeName);
-								}
-							}
-
-							return true;
+							const FString UnwantedType = GetNameSafe(Pin->PinType.PinSubCategoryObject.Get());
+							UE_LOG(LogBlueprint, Warning, TEXT("%s node ('%s') with a pin ('%s') that is now a deprecated type ('%s') - use the target's class instead (may just require a resave)."), *ReasonPrefix, *PinNodeName, *PinName, *UnwantedType);
+						}
+						else if (bUnwantedDefaultObject)
+						{
+							const FString UnwantedDefaultObject = GetNameSafe(Pin->DefaultObject);
+							UE_LOG(LogBlueprint, Warning, TEXT("%s node ('%s') with a pin ('%s') that has a default object ('%s') that is now deprecated - use the target's class instead (may just require a resave)."), *ReasonPrefix, *PinNodeName, *PinName, *UnwantedDefaultObject);
 						}
+						else
+						{
+							ensureMsgf(false, TEXT("Can not describe why the blueprint should be fixed."));
+							UE_LOG(LogBlueprint, Warning, TEXT("%s node ('%s') with a pin ('%s') that references a deprecated type - use the target's class instead (may just require a resave)."), *ReasonPrefix, *PinNodeName, *PinName);
+						}
+						return true;
 					}
 				}
 			}
 		}
-
-		return false;
 	}
 
-	static void HandleEditablePinNode(UK2Node_EditablePinBase* Node)
+	return false;
+}
+
+bool FLegacyEditorOnlyBlueprintUtils::IsUnwantedType(const FEdGraphPinType& Type)
+{
+	if (const UClass* BPClass = Cast<const UClass>(Type.PinSubCategoryObject.Get()))
 	{
-		check(Node);
+		if (BPClass->IsChildOf(UBlueprint::StaticClass()))
+		{
+			return true;
+		}
+	}
 
-		const UEdGraphSchema_K2* Schema = GetDefault<UEdGraphSchema_K2>();
+	return false;
+}
 
-		for (const TSharedPtr<FUserPinInfo>& UserPinInfo : Node->UserDefinedPins)
-		{
-			if (UserPinInfo.IsValid() && IsUnwantedType(UserPinInfo->PinType))
-			{
-				UserPinInfo->PinType.PinCategory = Schema->PC_Class;
-				UserPinInfo->PinType.PinSubCategoryObject = UObject::StaticClass();
+bool FLegacyEditorOnlyBlueprintUtils::IsUnwantedDefaultObject(const UObject* Obj)
+{
+	return Obj && Obj->IsA<UBlueprint>();
+}
 
-				UserPinInfo->PinDefaultValue.Empty();
-			}
-		}
+void FLegacyEditorOnlyBlueprintUtils::ChangePinType(FEdGraphPinType& Type)
+{
+	const UEdGraphSchema_K2* Schema = GetDefault<UEdGraphSchema_K2>();
 
-		for (UEdGraphPin* Pin : Node->Pins)
-		{
-			if (IsUnwantedType(Pin->PinType))
-			{
-				ChangePinType(Pin->PinType);
+	Type.PinCategory = Schema->PC_Class;
+	Type.PinSubCategoryObject = UObject::StaticClass();
+}
 
-				if (Pin->DefaultObject)
-				{
-					UBlueprint* DefaultBlueprint = Cast<UBlueprint>(Pin->DefaultObject);
-					Pin->DefaultObject = DefaultBlueprint ? *DefaultBlueprint->GeneratedClass : nullptr;
-				}
-			}
+void FLegacyEditorOnlyBlueprintUtils::HandleEditablePinNode(UK2Node_EditablePinBase* Node)
+{
+	check(Node);
 
+	const UEdGraphSchema_K2* Schema = GetDefault<UEdGraphSchema_K2>();
+
+	for (const TSharedPtr<FUserPinInfo>& UserPinInfo : Node->UserDefinedPins)
+	{
+		if (UserPinInfo.IsValid() && IsUnwantedType(UserPinInfo->PinType))
+		{
+			UserPinInfo->PinType.PinCategory = Schema->PC_Class;
+			UserPinInfo->PinType.PinSubCategoryObject = UObject::StaticClass();
+
+			UserPinInfo->PinDefaultValue.Empty();
 		}
 	}
 
-	static void HandleTemporaryVariableNode(UK2Node_TemporaryVariable* Node)
+	for (UEdGraphPin* Pin : Node->Pins)
 	{
-		check(Node);
-
-		UEdGraphPin* Pin = Node->GetVariablePin();
-		if (Pin && IsUnwantedType(Pin->PinType))
+		if (IsUnwantedType(Pin->PinType))
 		{
 			ChangePinType(Pin->PinType);
-			
+
 			if (Pin->DefaultObject)
 			{
 				UBlueprint* DefaultBlueprint = Cast<UBlueprint>(Pin->DefaultObject);
@@ -1301,92 +1367,48 @@ struct FEditoronlyBlueprintHelper
 			}
 		}
 
-		if (IsUnwantedType(Node->VariableType))
-		{
-			ChangePinType(Node->VariableType);
-		}
 	}
+}
 
-	static void HandleDefaultObjects(UK2Node* Node)
+void FLegacyEditorOnlyBlueprintUtils::HandleTemporaryVariableNode(UK2Node_TemporaryVariable* Node)
+{
+	check(Node);
+
+	UEdGraphPin* Pin = Node->GetVariablePin();
+	if (Pin && IsUnwantedType(Pin->PinType))
 	{
-		if (Node)
+		ChangePinType(Pin->PinType);
+			
+		if (Pin->DefaultObject)
 		{
-			for (UEdGraphPin* Pin : Node->Pins)
-			{
-				if (Pin && IsUnwantedDefaultObject(Pin->DefaultObject))
-				{
-					const UBlueprint* DefaultBlueprint = Cast<const UBlueprint>(Pin->DefaultObject);
-					if (DefaultBlueprint)
-					{
-						Pin->DefaultObject = *DefaultBlueprint->GeneratedClass;
-					}
-				}
-			}
+			UBlueprint* DefaultBlueprint = Cast<UBlueprint>(Pin->DefaultObject);
+			Pin->DefaultObject = DefaultBlueprint ? *DefaultBlueprint->GeneratedClass : nullptr;
 		}
 	}
 
-	static bool ChangeBlueprint(UBlueprint* Blueprint)
+	if (IsUnwantedType(Node->VariableType))
 	{
-		if (ShouldBeFixed(Blueprint, false))
-		{
-			UE_LOG(LogBlueprint, Log, TEXT("Bluepirnt references will be removed from '%s'"), *Blueprint->GetName());
-
-			const UEdGraphSchema_K2* Schema = GetDefault<UEdGraphSchema_K2>();
-
-			//any native function call with blueprint parameter
-			TArray<UEdGraph*> Graphs;
-			Blueprint->GetAllGraphs(Graphs);
-			for (UEdGraph* Graph : Graphs)
-			{
-				Schema->BackwardCompatibilityNodeConversion(Graph, false);
-			}
-
-			//any blueprint pin (function/event/macro parameter
-			for (const UEdGraph* Graph : Graphs)
-			{
-				{
-					TArray<UK2Node_EditablePinBase*> EditablePinNodes;
-					Graph->GetNodesOfClass(EditablePinNodes);
-					for (UK2Node_EditablePinBase* Node : EditablePinNodes)
-					{
-						HandleEditablePinNode(Node);
-					}
-				}
-
-				{
-					TArray<UK2Node_TemporaryVariable*> TempVariablesNodes;
-					Graph->GetNodesOfClass(TempVariablesNodes);
-					for (UK2Node_TemporaryVariable* Node : TempVariablesNodes)
-					{
-						HandleTemporaryVariableNode(Node);
-					}
-				}
-			}
-
-			// change variables
-			for (FBPVariableDescription& VarDesc : Blueprint->NewVariables)
-			{
-				if (IsUnwantedType(VarDesc.VarType))
-				{
-					ChangePinType(VarDesc.VarType);
-				}
-			}
+		ChangePinType(Node->VariableType);
+	}
+}
 
-			for (const UEdGraph* Graph : Graphs)
+void FLegacyEditorOnlyBlueprintUtils::HandleDefaultObjects(UK2Node* Node)
+{
+	if (Node)
+	{
+		for (UEdGraphPin* Pin : Node->Pins)
+		{
+			if (Pin && IsUnwantedDefaultObject(Pin->DefaultObject))
 			{
-				TArray<UK2Node*> Nodes;
-				Graph->GetNodesOfClass(Nodes);
-				for (UK2Node* Node : Nodes)
+				const UBlueprint* DefaultBlueprint = Cast<const UBlueprint>(Pin->DefaultObject);
+				if (DefaultBlueprint)
 				{
-					HandleDefaultObjects(Node);
+					Pin->DefaultObject = *DefaultBlueprint->GeneratedClass;
 				}
 			}
-
-			return true;
 		}
-		return false;
 	}
-};
+}
 
 /**
 	Procedure used to remove old function implementations and child properties from data only blueprints.
@@ -1522,11 +1544,7 @@ UClass* FBlueprintEditorUtils::RegenerateBlueprintClass(UBlueprint* Blueprint, U
 
 		bool bSkeletonUpToDate = FKismetEditorUtilities::GenerateBlueprintSkeleton(Blueprint);
 
-		static FBoolConfigValueHelper ReplaceBlueprintWithClass(TEXT("EditoronlyBP"), TEXT("bReplaceBlueprintWithClass"));
-		if (ReplaceBlueprintWithClass)
-		{
-			FEditoronlyBlueprintHelper::ChangeBlueprint(Blueprint);
-		}
+		FLegacyEditorOnlyBlueprintUtils::FixupBlueprint(Blueprint);
 
 		const bool bDataOnlyClassThatMustBeRecompiled = !bHasCode && !bIsMacro
 			&& (!ClassToRegenerate || (Blueprint->ParentClass != ClassToRegenerate->GetSuperClass()));
@@ -1616,11 +1634,6 @@ UClass* FBlueprintEditorUtils::RegenerateBlueprintClass(UBlueprint* Blueprint, U
 			// Flag data only blueprints as being up-to-date
 			Blueprint->Status = BS_UpToDate;
 		}
-
-		if (ReplaceBlueprintWithClass)
-		{
-			FEditoronlyBlueprintHelper::ShouldBeFixed(Blueprint, true);
-		}
 		
 		// Patch the new CDOs to the old indices in the linker
 		if( Blueprint->SkeletonGeneratedClass )
diff --git a/Engine/Source/Editor/UnrealEd/Private/Kismet2/Kismet2.cpp b/Engine/Source/Editor/UnrealEd/Private/Kismet2/Kismet2.cpp
--- a/Engine/Source/Editor/UnrealEd/Private/Kismet2/Kismet2.cpp
+++ b/Engine/Source/Editor/UnrealEd/Private/Kismet2/Kismet2.cpp
@@ -634,6 +634,13 @@ UK2Node_Event* FKismetEditorUtilities::AddDefaultEventNode(UBlueprint* InBluepri
 			ParentFunctionNode->SetFromFunction(ValidParent);
 			ParentFunctionNode->AllocateDefaultPins();
 
+			for (UEdGraphPin* EventPin : NewEventNode->Pins)
+			{
+				if (UEdGraphPin* ParentPin = ParentFunctionNode->FindPin(EventPin->PinName, EGPD_Input))
+				{
+					ParentPin->MakeLinkTo(EventPin);
+				}
+			}
 			ParentFunctionNode->GetExecPin()->MakeLinkTo(NewEventNode->FindPin(Schema->PN_Then));
 
 			ParentFunctionNode->NodePosX = FunctionFromNode.Node->NodePosX + FunctionFromNode.Node->NodeWidth + 200;
@@ -949,11 +956,6 @@ void FKismetEditorUtilities::RecompileBlueprintBytecode(UBlueprint* BlueprintObj
 
 	FKismetCompilerOptions CompileOptions;
 	CompileOptions.CompileType = EKismetCompileType::BytecodeOnly;
-
-	// Determine if we want profiling data
-	IBlueprintProfilerInterface* ProfilerInterface = FModuleManager::GetModulePtr<IBlueprintProfilerInterface>("BlueprintProfiler");
-	CompileOptions.bAddInstrumentation = ProfilerInterface && ProfilerInterface->IsProfilerEnabled() && GetDefault<UEditorExperimentalSettings>()->bBlueprintPerformanceAnalysisTools;
-
 	{
 		FRecreateUberGraphFrameScope RecreateUberGraphFrameScope(BlueprintObj->GeneratedClass, true);
 		Compiler.CompileBlueprint(BlueprintObj, CompileOptions, Results, NULL, ObjLoaded);
diff --git a/Engine/Source/Editor/UnrealEd/Private/UserDefinedStructEditorData.cpp b/Engine/Source/Editor/UnrealEd/Private/UserDefinedStructEditorData.cpp
--- a/Engine/Source/Editor/UnrealEd/Private/UserDefinedStructEditorData.cpp
+++ b/Engine/Source/Editor/UnrealEd/Private/UserDefinedStructEditorData.cpp
@@ -4,6 +4,7 @@
 #include "StructureEditorUtils.h"
 #include "Engine/UserDefinedStruct.h"
 #include "BlueprintEditorUtils.h"
+#include "Blueprint/BlueprintSupport.h"
 
 #define LOCTEXT_NAMESPACE "UserDefinedStructEditorData"
 
@@ -114,6 +115,11 @@ void UUserDefinedStructEditorData::RecreateDefaultInstance(FString* OutLog)
 	ensure(DefaultStructInstance.IsValid() && DefaultStructInstance.GetStruct() == ScriptStruct);
 	if (DefaultStructInstance.IsValid() && StructData && ScriptStruct)
 	{
+		// When loading, the property's default value may end up being filled with a placeholder. 
+		// This tracker object allows the linker to track the actual object that is being filled in 
+		// so it can calculate an offset to the property and write in the placeholder value:
+		FScopedPlaceholderRawContainerTracker TrackDefaultObject(StructData);
+
 		DefaultStructInstance.SetPackage(ScriptStruct->GetOutermost());
 
 		for (TFieldIterator<UProperty> It(ScriptStruct); It; ++It)
diff --git a/Engine/Source/Editor/UnrealEd/Public/EdGraphUtilities.h b/Engine/Source/Editor/UnrealEd/Public/EdGraphUtilities.h
--- a/Engine/Source/Editor/UnrealEd/Public/EdGraphUtilities.h
+++ b/Engine/Source/Editor/UnrealEd/Public/EdGraphUtilities.h
@@ -127,6 +127,9 @@ public:
 	/** returns true if the ParameterName is marked as a SetParam in the relevant Function */
 	static bool IsSetParam(const UFunction* Function, const FString& ParameterName);
 
+	/** returns true if the ParameterName is marked as a MapParam, MapKeyParam, or MapValueParam in the relevant Function */
+	static bool IsMapParam(const UFunction* Function, const FString& ParameterName);
+
 private:
 	static TArray< TSharedPtr<FGraphPanelNodeFactory> > VisualNodeFactories;
 	static TArray< TSharedPtr<FGraphPanelPinFactory> > VisualPinFactories;
diff --git a/Engine/Source/Editor/UnrealEd/Public/Kismet2/BlueprintEditorUtils.h b/Engine/Source/Editor/UnrealEd/Public/Kismet2/BlueprintEditorUtils.h
--- a/Engine/Source/Editor/UnrealEd/Public/Kismet2/BlueprintEditorUtils.h
+++ b/Engine/Source/Editor/UnrealEd/Public/Kismet2/BlueprintEditorUtils.h
@@ -7,6 +7,7 @@
 #include "EdGraphSchema_K2.h"
 #include "K2Node_EditablePinBase.h"
 #include "Engine/LevelScriptBlueprint.h"
+#include "BlueprintSupport.h" // for FLegacyEditorOnlyBlueprintOptions
 
 class  USCS_Node;
 struct FComponentKey;
@@ -1471,4 +1472,25 @@ struct UNREALED_API FMakeClassSpawnableOnScope
 			}
 		}
 	}
+};
+
+/** 
+ * Temporary util struct that deals with the now deprecated [EditoronlyBP] 
+ * settings. Inherited from FLegacyEditorOnlyBlueprintOptions (in CoreUObject), 
+ * which consolidates these settings, but is extended here to deal with UnrealEd 
+ * specific types.
+ */
+struct FLegacyEditorOnlyBlueprintUtils : public FLegacyEditorOnlyBlueprintOptions
+{
+	UNREALED_API static bool DoPinsMatch(const FEdGraphPinType& Input, const FEdGraphPinType& Output);
+	static bool FixupBlueprint(UBlueprint* Blueprint);
+
+private:
+	static bool HasLegacyBlueprintReferences(UBlueprint* Blueprint);
+	static bool IsUnwantedType(const FEdGraphPinType& Type);
+	static bool IsUnwantedDefaultObject(const UObject* Obj);
+	static void ChangePinType(FEdGraphPinType& Type);
+	static void HandleEditablePinNode(class UK2Node_EditablePinBase* Node);
+	static void HandleTemporaryVariableNode(class UK2Node_TemporaryVariable* Node);
+	static void HandleDefaultObjects(UK2Node* Node);
 };
\ No newline at end of file
diff --git a/Engine/Source/Runtime/Core/Public/Containers/Set.h b/Engine/Source/Runtime/Core/Public/Containers/Set.h
--- a/Engine/Source/Runtime/Core/Public/Containers/Set.h
+++ b/Engine/Source/Runtime/Core/Public/Containers/Set.h
@@ -1492,7 +1492,7 @@ public:
 		return !bIsAlreadyInSet;
 	}
 
-	void Remove(const void* Element, const FScriptSetLayout& Layout, TFunctionRef<uint32(const void*)> GetKeyHash, TFunctionRef<bool(const void*, const void*)> EqualityFn)
+	bool Remove(const void* Element, const FScriptSetLayout& Layout, TFunctionRef<uint32(const void*)> GetKeyHash, TFunctionRef<bool(const void*, const void*)> EqualityFn, TFunctionRef<void(void*)> DestroyFn )
 	{
 		if (Elements.Num())
 		{
@@ -1507,11 +1507,13 @@ public:
 				CurrentElement = (uint8*)Elements.GetData(ElementId, Layout.SparseArrayLayout);
 				if(EqualityFn(Element, CurrentElement))
 				{
+					DestroyFn(CurrentElement);
 					RemoveAt(ElementId, Layout);
-					return;
+					return true;
 				}
 			}
 		}
+		return false;
 	}
 
 private:
diff --git a/Engine/Source/Runtime/Core/Public/UObject/ObjectVersion.h b/Engine/Source/Runtime/Core/Public/UObject/ObjectVersion.h
--- a/Engine/Source/Runtime/Core/Public/UObject/ObjectVersion.h
+++ b/Engine/Source/Runtime/Core/Public/UObject/ObjectVersion.h
@@ -612,6 +612,8 @@ enum EUnrealEngineObjectUE4Version
 	VER_UE4_PRELOAD_DEPENDENCIES_IN_COOKED_EXPORTS,
 	// Cooked files contain the TemplateIndex used by the event driven loader (the serialization is largely independent of the use of the new loader, i.e. this will be null if cooking for the old loader)
 	VER_UE4_TemplateIndex_IN_COOKED_EXPORTS,
+	// FPropertyTag includes contained type(s) for Set and Map properties:
+	VER_UE4_PROPERTY_TAG_SET_MAP_SUPPORT,
 	
 	// -----<new versions can be added before this line>-------------------------------------------------
 	// - this needs to be the last line (see note below)
diff --git a/Engine/Source/Runtime/CoreUObject/Private/Blueprint/BlueprintSupport.cpp b/Engine/Source/Runtime/CoreUObject/Private/Blueprint/BlueprintSupport.cpp
--- a/Engine/Source/Runtime/CoreUObject/Private/Blueprint/BlueprintSupport.cpp
+++ b/Engine/Source/Runtime/CoreUObject/Private/Blueprint/BlueprintSupport.cpp
@@ -1737,7 +1737,6 @@ void FLinkerLoad::CreateDynamicTypeLoader()
 		OuterImport->ClassPackage = GLongCoreUObjectPackageName;
 		OuterImport->ObjectName = Import.PackageName;
 
-#if !USE_EVENT_DRIVEN_ASYNC_LOAD
 		if((Import.ClassName == DynamicClassName) && (Import.ClassPackageName == DynamicClassPackageName))
 		{
 			const FString DynamicClassPath = Import.PackageName.ToString() + TEXT(".") + Import.ObjectName.ToString();
@@ -1750,7 +1749,6 @@ void FLinkerLoad::CreateDynamicTypeLoader()
 				ClassConstructFn->StaticClassFn();
 			}
 		}
-#endif
 	}
 
 	// Create Export
@@ -2229,3 +2227,178 @@ void IBlueprintNativeCodeGenCore::Register(const IBlueprintNativeCodeGenCore* Co
 }
 
 #endif // WITH_EDITOR
+
+
+/*******************************************************************************
+ * FLegacyEditorOnlyBlueprintManager
+ ******************************************************************************/
+
+struct FLegacyEditorOnlyBlueprintOptions_Impl
+{
+	static void Init();
+	static const UClass* GetUBlueprintObjType();
+
+	static bool bExcludeBlueprintObjectsFromCookedBuilds;
+	static bool bFixupLegacyBlueprintProperties;
+	static bool bProhibitLegacyBlueprintVarType;
+	static bool bForceAllowLegacyBlueprintPinConnections;
+};
+
+bool FLegacyEditorOnlyBlueprintOptions_Impl::bExcludeBlueprintObjectsFromCookedBuilds = false;
+bool FLegacyEditorOnlyBlueprintOptions_Impl::bFixupLegacyBlueprintProperties = false;
+bool FLegacyEditorOnlyBlueprintOptions_Impl::bProhibitLegacyBlueprintVarType = false;
+bool FLegacyEditorOnlyBlueprintOptions_Impl::bForceAllowLegacyBlueprintPinConnections = false;
+
+void FLegacyEditorOnlyBlueprintOptions_Impl::Init()
+{
+	static bool bInited = false;
+	if (!bInited)
+	{
+		const TCHAR* SectionId = TEXT("EditoronlyBP");
+
+		bool bHasSettings = GConfig->GetBool(SectionId, TEXT("bDontLoadBlueprintOutsideEditor"), bExcludeBlueprintObjectsFromCookedBuilds, GEditorIni);
+		bHasSettings |= GConfig->GetBool(SectionId, TEXT("bReplaceBlueprintWithClass"),   bFixupLegacyBlueprintProperties, GEditorIni);
+		bHasSettings |= GConfig->GetBool(SectionId, TEXT("bBlueprintIsNotBlueprintType"), bProhibitLegacyBlueprintVarType, GEditorIni);
+		bHasSettings |= GConfig->GetBool(SectionId, TEXT("bAllowClassAndBlueprintPinMatching"), bForceAllowLegacyBlueprintPinConnections, GEditorIni);
+
+		// we expect all these settings to be enabled
+		const bool bOptionsDisabled = !bExcludeBlueprintObjectsFromCookedBuilds || !bFixupLegacyBlueprintProperties ||
+			!bProhibitLegacyBlueprintVarType || !bForceAllowLegacyBlueprintPinConnections;
+		if (bOptionsDisabled)
+		{
+			if (bHasSettings)
+			{
+				UE_LOG(LogBlueprintSupport, Warning, TEXT("Editor config [EditoronlyBP] settings are DEPRECATED. You're explicitly disabling some. If you're relying on this functionality, then please fix it up so your Blueprints don't break in a future release."));
+			}
+			else
+			{
+				UE_LOG(LogBlueprintSupport, Warning, TEXT("Editor config [EditoronlyBP] settings are DEPRECATED. You're not explicitly setting any, but they've been defaulting to off. If you are relying on this functionality, then please fix it up so your Blueprints don't break in a future release."));
+			}
+		}
+
+		bInited = true;
+	}
+}
+
+const UClass* FLegacyEditorOnlyBlueprintOptions_Impl::GetUBlueprintObjType()
+{
+	static const UClass* BlueprintClass = FindObject<UClass>(nullptr, TEXT("/Script/Engine.Blueprint"));
+	return BlueprintClass;
+}
+
+FString FLegacyEditorOnlyBlueprintOptions::GetDefaultEditorConfig()
+{
+	FString ConfigStr;
+	ConfigStr += TEXT("[EditoronlyBP]") LINE_TERMINATOR;
+	ConfigStr += TEXT("bAllowClassAndBlueprintPinMatching=true") LINE_TERMINATOR;
+	ConfigStr += TEXT("bReplaceBlueprintWithClass=true") LINE_TERMINATOR;
+	ConfigStr += TEXT("bDontLoadBlueprintOutsideEditor=true") LINE_TERMINATOR;
+	ConfigStr += TEXT("bBlueprintIsNotBlueprintType=true") LINE_TERMINATOR;
+	return ConfigStr;
+}
+
+bool FLegacyEditorOnlyBlueprintOptions::IncludeUBlueprintObjsInCookedBuilds()
+{
+	FLegacyEditorOnlyBlueprintOptions_Impl::Init();
+	return !FLegacyEditorOnlyBlueprintOptions_Impl::bExcludeBlueprintObjectsFromCookedBuilds;
+}
+
+bool FLegacyEditorOnlyBlueprintOptions::AllowLegacyBlueprintPinMatchesWithClass()
+{
+	FLegacyEditorOnlyBlueprintOptions_Impl::Init();
+	return FLegacyEditorOnlyBlueprintOptions_Impl::bForceAllowLegacyBlueprintPinConnections;
+}
+
+bool FLegacyEditorOnlyBlueprintOptions::FixupLegacyBlueprintReferences()
+{
+	FLegacyEditorOnlyBlueprintOptions_Impl::Init();
+	return FLegacyEditorOnlyBlueprintOptions_Impl::bFixupLegacyBlueprintProperties;
+}
+
+bool FLegacyEditorOnlyBlueprintOptions::DetectInvalidBlueprintExport(const FLinkerSave* Linker, const int32 ExportIndex)
+{
+	const UObject* ExportObj = Linker->ExportMap[ExportIndex].Object;
+	if (Linker->IsCooking() && ExportObj && IncludeUBlueprintObjsInCookedBuilds())
+	{
+		const UClass* BlueprintClass = FLegacyEditorOnlyBlueprintOptions_Impl::GetUBlueprintObjType();
+		if (BlueprintClass && ExportObj->IsA(BlueprintClass))
+		{
+			const bool bIsFatal = IsEventDrivenLoaderEnabledInCookedBuilds();
+			if (bIsFatal)
+			{
+				UE_LOG(LogBlueprintSupport, Fatal,
+					TEXT("You're exporting a UBlueprint ('%s') editor object with your cook. This is now DEPRECATED (the Blueprint's class & CDO should be all that's needed). This is incompatible with the new event-driven loader. Please add the following to your editor ini file to strip UBluprints from cooked packages (these settings will become the default):\n%s"),
+					*ExportObj->GetName(),
+					*GetDefaultEditorConfig());
+			}
+			else
+			{
+				UE_LOG(LogBlueprintSupport, Warning,
+					TEXT("You're exporting a UBlueprint ('%s') editor object with your cook. This is now DEPRECATED (the Blueprint's class & CDO should be all that's needed). Please fix up any relience you have on it, and add the following to your editor ini file (these settings will become the default):\n%s"),
+					*ExportObj->GetName(),
+					*GetDefaultEditorConfig());
+			}
+			return true;
+		}
+	}
+	return false;
+}
+
+bool FLegacyEditorOnlyBlueprintOptions::FixupClassProperty(const UClassProperty* Property, void* Value)
+{
+	const UObject* ObjValue = Property->GetObjectPropertyValue(Value);
+	if (ObjValue && (Property->MetaClass == UObject::StaticClass()))
+	{
+		// since we're in core, we don't have access to UBlueprint::StaticClass()
+		const UClass* BlueprintClass = FLegacyEditorOnlyBlueprintOptions_Impl::GetUBlueprintObjType();
+		if (BlueprintClass && ObjValue->IsA(BlueprintClass))
+		{
+			FLegacyEditorOnlyBlueprintOptions_Impl::Init();
+			if (FLegacyEditorOnlyBlueprintOptions_Impl::bFixupLegacyBlueprintProperties)
+			{
+				UE_LOG(LogBlueprintSupport, Warning,
+					TEXT("Attempting to fix up an old Blueprint property (%s), which is referencing a UBlueprint (%s) instead of its class - this fixup is DEPRECATED; please fix manually (possibly just a resave)"),
+					*Property->GetFullName(),
+					*ObjValue->GetFullName());
+
+				// again, no access to UBlueprint, so we have to sift through reflection data to get what we want
+				if (UClassProperty* BPGeneratedClassProp = FindField<UClassProperty>(BlueprintClass, TEXT("GeneratedClass")))
+				{
+					UObject* ClassObject = BPGeneratedClassProp->GetPropertyValue_InContainer(ObjValue);
+					Property->SetObjectPropertyValue(Value, ClassObject);
+					return true;
+				}
+			}
+			else
+			{
+				UE_LOG(LogBlueprintSupport, Warning,
+					TEXT("An old Blueprint property (%s), references a UBlueprint (%s) instead of its class - this is DEPRECATED in packaged builds; please fix (possibly by enabling bReplaceBlueprintWithClass and resaving)"),
+					*Property->GetFullName(),
+					*ObjValue->GetFullName());
+			}
+		}
+	}
+
+	return false;
+}
+
+bool FLegacyEditorOnlyBlueprintOptions::IsTypeProhibited(const UClass* VarType)
+{
+	// since we're in core, we don't have access to UBlueprint::StaticClass()
+	const UClass* BlueprintClass = FLegacyEditorOnlyBlueprintOptions_Impl::GetUBlueprintObjType();
+	if (BlueprintClass && VarType->IsChildOf(BlueprintClass))
+	{
+		FLegacyEditorOnlyBlueprintOptions_Impl::Init();
+		if (FLegacyEditorOnlyBlueprintOptions_Impl::bForceAllowLegacyBlueprintPinConnections)
+		{
+			return true;
+		}
+	}
+	return false;
+}
+
+
+
+
+
+
diff --git a/Engine/Source/Runtime/CoreUObject/Private/UObject/Class.cpp b/Engine/Source/Runtime/CoreUObject/Private/UObject/Class.cpp
--- a/Engine/Source/Runtime/CoreUObject/Private/UObject/Class.cpp
+++ b/Engine/Source/Runtime/CoreUObject/Private/UObject/Class.cpp
@@ -4883,6 +4883,19 @@ UObject* UDynamicClass::FindArchetype(UClass* ArchetypeClass, const FName Archet
 {
 	UDynamicClass* ThisClass = const_cast<UDynamicClass*>(this);
 	UObject* Archetype = static_cast<UObject*>(FindObjectWithOuter(ThisClass, ArchetypeClass, ArchetypeName));
+	if (!Archetype)
+	{
+		// See UBlueprintGeneratedClass::FindArchetype, UE-35259, UE-37480
+		const FName ArchetypeBaseName = FName(ArchetypeName, 0);
+		if (ArchetypeBaseName != ArchetypeName)
+		{
+			UObject* const* FountComponentTemplate = ComponentTemplates.FindByPredicate([&](UObject* InObj) -> bool
+			{ 
+				return InObj && (InObj->GetFName() == ArchetypeBaseName) && InObj->IsA(ArchetypeClass);
+			});
+			Archetype = FountComponentTemplate ? *FountComponentTemplate : nullptr;
+		}
+	}
 	const UClass* SuperClass = GetSuperClass();
 	return Archetype ? Archetype :
 		(SuperClass ? SuperClass->FindArchetype(ArchetypeClass, ArchetypeName) : nullptr);
diff --git a/Engine/Source/Runtime/CoreUObject/Private/UObject/LinkerPlaceholderBase.cpp b/Engine/Source/Runtime/CoreUObject/Private/UObject/LinkerPlaceholderBase.cpp
--- a/Engine/Source/Runtime/CoreUObject/Private/UObject/LinkerPlaceholderBase.cpp
+++ b/Engine/Source/Runtime/CoreUObject/Private/UObject/LinkerPlaceholderBase.cpp
@@ -18,6 +18,7 @@
 struct FPlaceholderContainerTracker : TThreadSingleton<FPlaceholderContainerTracker>
 {	
 	TArray<UObject*> PerspectiveReferencerStack;
+	TArray<void*> PerspectiveRootDataStack;
 	// as far as I can tell, structs are going to be the only bridging point 
 	// between property ownership
 	TArray<const UStructProperty*> IntermediatePropertyStack; 
@@ -56,6 +57,7 @@ public:
 	 * @return The UObject instance that is assumably referencing a placeholder (null if one couldn't be found)
 	 */
 	static UObject* FindPlaceholderContainer(const FLinkerPlaceholderBase::FPlaceholderValuePropertyPath& PropertyChainRef);
+	static void* FindRawPlaceholderContainer(const FLinkerPlaceholderBase::FPlaceholderValuePropertyPath& PropertyChainRef);
 };
 
 //------------------------------------------------------------------------------
@@ -153,6 +155,12 @@ UObject* FLinkerPlaceholderObjectImpl::FindPlaceholderContainer(const FLinkerPla
 	return ContainerObj;
 }
 
+void* FLinkerPlaceholderObjectImpl::FindRawPlaceholderContainer(const FLinkerPlaceholderBase::FPlaceholderValuePropertyPath& PropertyChainRef)
+{
+	TArray<void*>& PossibleStructReferencers = FPlaceholderContainerTracker::Get().PerspectiveRootDataStack;
+	return PossibleStructReferencers[0];
+}
+
 /*******************************************************************************
  * FPlaceholderContainerTracker / FScopedPlaceholderPropertyTracker
  ******************************************************************************/
@@ -173,12 +181,30 @@ FScopedPlaceholderContainerTracker::~FScopedPlaceholderContainerTracker()
 #endif // USE_DEFERRED_DEPENDENCY_CHECK_VERIFICATION_TESTS
 }
 
+#if WITH_EDITOR
+
+FScopedPlaceholderRawContainerTracker::FScopedPlaceholderRawContainerTracker(void* InData)
+	:Data(InData)
+{
+	FPlaceholderContainerTracker::Get().PerspectiveRootDataStack.Add(InData);
+}
+
+FScopedPlaceholderRawContainerTracker::~FScopedPlaceholderRawContainerTracker()
+{
+	void* StackTop = FPlaceholderContainerTracker::Get().PerspectiveRootDataStack.Pop();
+#if USE_DEFERRED_DEPENDENCY_CHECK_VERIFICATION_TESTS
+	check(StackTop == Data);
+#endif // USE_DEFERRED_DEPENDENCY_CHECK_VERIFICATION_TESTS
+}
+
+#endif
+
 //------------------------------------------------------------------------------
 FScopedPlaceholderPropertyTracker::FScopedPlaceholderPropertyTracker(const UStructProperty* InIntermediateProperty)
 	: IntermediateProperty(nullptr) // leave null, as a sentinel value (implying that PerspectiveReferencerStack was empty)
 {
 	FPlaceholderContainerTracker& ContainerRepository = FPlaceholderContainerTracker::Get();
-	if (ContainerRepository.PerspectiveReferencerStack.Num() > 0)
+	if (ContainerRepository.PerspectiveReferencerStack.Num() > 0 || ContainerRepository.PerspectiveRootDataStack.Num() > 0)
 	{
 		IntermediateProperty = InIntermediateProperty;
 		ContainerRepository.IntermediatePropertyStack.Add(InIntermediateProperty);
@@ -202,6 +228,7 @@ FScopedPlaceholderPropertyTracker::~FScopedPlaceholderPropertyTracker()
 	{
 		check(ContainerRepository.IntermediatePropertyStack.Num()  == 0);
 		check(ContainerRepository.PerspectiveReferencerStack.Num() == 0);
+		check(ContainerRepository.PerspectiveRootDataStack.Num() == 0);
 #endif // USE_DEFERRED_DEPENDENCY_CHECK_VERIFICATION_TESTS
 	}
 }
@@ -230,7 +257,7 @@ FLinkerPlaceholderBase::FPlaceholderValuePropertyPath::FPlaceholderValueProperty
 		// to help trace the property path)
 		else if (const UScriptStruct* StructOwner = Cast<const UScriptStruct>(PropertyOuter))
 		{
-			if (DEFERRED_DEPENDENCY_ENSURE(StructStackIndex != INDEX_NONE))
+			if (StructStackIndex != INDEX_NONE)
 			{
 				// we expect the top struct property to be the one we're currently serializing
 				const UStructProperty* SerializingStructProp = StructPropertyStack[StructStackIndex];
@@ -249,6 +276,11 @@ FLinkerPlaceholderBase::FPlaceholderValuePropertyPath::FPlaceholderValueProperty
 				}
 				--StructStackIndex;
 			}
+			else
+			{
+				// We're serializing a struct that isn't owned by a UObject (e.g. UUserDefinedStructEditorData::DefaultStructInstance)
+				break;
+			}
 		}
 		PropertyOuter = PropertyOuter->GetOuter();
 	}
@@ -286,6 +318,13 @@ int32 FLinkerPlaceholderBase::FPlaceholderValuePropertyPath::Resolve(FLinkerPlac
 	uint8* OutermostAddress = OutermostProperty->ContainerPtrToValuePtr<uint8>((uint8*)Container, /*ArrayIndex =*/0);
 	return FLinkerPlaceholderObjectImpl::ResolvePlaceholderValues(PropertyChain, PropertyChain.Num() - 1, OutermostAddress, Placeholder->GetPlaceholderAsUObject(), Replacement);
 }
+
+int32 FLinkerPlaceholderBase::FPlaceholderValuePropertyPath::ResolveRaw(FLinkerPlaceholderBase* Placeholder, UObject* Replacement, void* Container) const
+{
+	const UProperty* OutermostProperty = PropertyChain.Last();
+	uint8* OutermostAddress = OutermostProperty->ContainerPtrToValuePtr<uint8>((uint8*)Container, /*ArrayIndex =*/0);
+	return FLinkerPlaceholderObjectImpl::ResolvePlaceholderValues(PropertyChain, PropertyChain.Num() - 1, OutermostAddress, Placeholder->GetPlaceholderAsUObject(), Replacement);;
+}
  
 /*******************************************************************************
  * FLinkerPlaceholderBase
@@ -310,24 +349,34 @@ bool FLinkerPlaceholderBase::AddReferencingPropertyValue(const UObjectProperty*
 {
 	FPlaceholderValuePropertyPath PropertyChain(ReferencingProperty);
 	UObject* ReferencingContainer = FLinkerPlaceholderObjectImpl::FindPlaceholderContainer(PropertyChain);
-
+	if (ReferencingContainer != nullptr)
+	{
 #if USE_DEFERRED_DEPENDENCY_CHECK_VERIFICATION_TESTS
-	check(ReferencingProperty->GetObjectPropertyValue(DataPtr) == GetPlaceholderAsUObject());
-	check(PropertyChain.IsValid());
-	checkf(ReferencingContainer != nullptr, TEXT("We couldn't reliably determine the object that a placeholder value should belong to - are we missing a FScopedPlaceholderContainerTracker somewhere?"));
+		check(ReferencingProperty->GetObjectPropertyValue(DataPtr) == GetPlaceholderAsUObject());
+		check(PropertyChain.IsValid());
 #endif // USE_DEFERRED_DEPENDENCY_CHECK_VERIFICATION_TESTS
 
-	if (ReferencingContainer != nullptr)
+		if (ReferencingContainer != nullptr)
+		{
+			ReferencingContainers.FindOrAdd(ReferencingContainer).Add(PropertyChain);
+		}
+		return (ReferencingContainer != nullptr);
+	}
+	else
 	{
-		ReferencingContainers.FindOrAdd(ReferencingContainer).Add(PropertyChain);
+		void* ReferencingRootStruct = FLinkerPlaceholderObjectImpl::FindRawPlaceholderContainer(PropertyChain);
+		if (ReferencingRootStruct)
+		{
+			ReferencingRawContainers.FindOrAdd(ReferencingRootStruct).Add(PropertyChain);
+		}
+		return ReferencingRootStruct != nullptr;
 	}
-	return (ReferencingContainer != nullptr);
 }
 
 //------------------------------------------------------------------------------
 bool FLinkerPlaceholderBase::HasKnownReferences() const
 {
-	return (ReferencingContainers.Num() > 0);
+	return (ReferencingContainers.Num() > 0) || ReferencingRawContainers.Num() > 0;
 }
 
 //------------------------------------------------------------------------------
@@ -335,6 +384,7 @@ int32 FLinkerPlaceholderBase::ResolveAllPlaceholderReferences(UObject* Replaceme
 {
 	int32 ReplacementCount = ResolvePlaceholderPropertyValues(ReplacementObj);
 	ReferencingContainers.Empty();
+	ReferencingRawContainers.Empty();
 
 	MarkAsResolved();
 	return ReplacementCount;
@@ -398,6 +448,21 @@ int32 FLinkerPlaceholderBase::ResolvePlaceholderPropertyValues(UObject* NewObjec
 		}
 	}
 
+	for (auto& ReferencingRawPair : ReferencingRawContainers)
+	{
+		void* RawValue = ReferencingRawPair.Key;
+		check(RawValue);
+
+		for (const FPlaceholderValuePropertyPath& PropertyRef : ReferencingRawPair.Value)
+		{
+			int32 ResolvedCount = PropertyRef.ResolveRaw(this, NewObjectValue, RawValue);
+			ResolvedTotal += ResolvedCount;
+#if USE_DEFERRED_DEPENDENCY_CHECK_VERIFICATION_TESTS
+			check(ResolvedCount > 0);
+#endif
+		}
+	}
+
 	return ResolvedTotal;
 }
 
diff --git a/Engine/Source/Runtime/CoreUObject/Private/UObject/LinkerPlaceholderBase.h b/Engine/Source/Runtime/CoreUObject/Private/UObject/LinkerPlaceholderBase.h
--- a/Engine/Source/Runtime/CoreUObject/Private/UObject/LinkerPlaceholderBase.h
+++ b/Engine/Source/Runtime/CoreUObject/Private/UObject/LinkerPlaceholderBase.h
@@ -20,7 +20,7 @@ class UObjectProperty;
 struct FScopedPlaceholderContainerTracker
 {
 public:
-	 FScopedPlaceholderContainerTracker(UObject* PerspectivePlaceholderReferencer);
+	FScopedPlaceholderContainerTracker(UObject* PerspectivePlaceholderReferencer);
 	~FScopedPlaceholderContainerTracker();
 
 private:
@@ -167,6 +167,7 @@ private:
 		 * @return The number of values successfully replaced (could be multiple for array properties).
 		 */
 		int32 Resolve(FLinkerPlaceholderBase* Placeholder, UObject* Replacement, UObject* Container) const;
+		int32 ResolveRaw(FLinkerPlaceholderBase* Placeholder, UObject* Replacement, void* Container) const;
 
 	private:
 		/** Denotes the property hierarchy used to reach this leaf property that is referencing a placeholder*/
@@ -206,6 +207,7 @@ private:
 
 	/** Tracks container objects that have property values set to reference this placeholder (references that need to be replaced later) */
 	TMap< TWeakObjectPtr<UObject>, FReferencingPropertySet > ReferencingContainers;
+	TMap< void*, FReferencingPropertySet > ReferencingRawContainers;
 };
 
 /*******************************************************************************
diff --git a/Engine/Source/Runtime/CoreUObject/Private/UObject/MetaData.cpp b/Engine/Source/Runtime/CoreUObject/Private/UObject/MetaData.cpp
--- a/Engine/Source/Runtime/CoreUObject/Private/UObject/MetaData.cpp
+++ b/Engine/Source/Runtime/CoreUObject/Private/UObject/MetaData.cpp
@@ -494,12 +494,9 @@ void UMetaData::InitializeRedirectMap()
 	}
 }
 
-#if HACK_HEADER_GENERATOR
-
 FName UMetaData::GetRemappedKeyName(FName OldKey)
 {
 	InitializeRedirectMap();
 	return KeyRedirectMap.FindRef(OldKey);
 }
 
-#endif
diff --git a/Engine/Source/Runtime/CoreUObject/Private/UObject/PropertyBaseObject.cpp b/Engine/Source/Runtime/CoreUObject/Private/UObject/PropertyBaseObject.cpp
--- a/Engine/Source/Runtime/CoreUObject/Private/UObject/PropertyBaseObject.cpp
+++ b/Engine/Source/Runtime/CoreUObject/Private/UObject/PropertyBaseObject.cpp
@@ -414,6 +414,11 @@ UObject* UObjectPropertyBase::FindImportedObject( const UProperty* Property, UOb
 		{
 #if USE_CIRCULAR_DEPENDENCY_LOAD_DEFERRING
 			FLinkerLoad* Linker = (OwnerObject != nullptr) ? OwnerObject->GetClass()->GetLinker() : nullptr;
+			if (Linker == nullptr)
+			{
+				// Fall back on the Properties owner. That is probably the thing that has triggered this load:
+				Linker = Property->GetLinker();
+			}
 			const bool bDeferAssetImports = (Linker != nullptr) && (Linker->LoadFlags & LOAD_DeferDependencyLoads);
 
 			if (bDeferAssetImports)
diff --git a/Engine/Source/Runtime/CoreUObject/Private/UObject/PropertyClass.cpp b/Engine/Source/Runtime/CoreUObject/Private/UObject/PropertyClass.cpp
--- a/Engine/Source/Runtime/CoreUObject/Private/UObject/PropertyClass.cpp
+++ b/Engine/Source/Runtime/CoreUObject/Private/UObject/PropertyClass.cpp
@@ -3,6 +3,7 @@
 #include "CoreUObjectPrivate.h"
 #include "PropertyHelper.h"
 #include "LinkerPlaceholderClass.h"
+#include "BlueprintSupport.h" // for FLegacyEditorOnlyBlueprintManager::FixupClassProperty()
 
 /*-----------------------------------------------------------------------------
 	UClassProperty.
@@ -143,51 +144,9 @@ bool UClassProperty::SameType(const UProperty* Other) const
 void UClassProperty::CheckValidObject(void* Value) const
 {
 #if WITH_EDITOR
-	// Ugly hack to replace Blueprint references with Class references.
-
-	struct FReplaceBlueprintWithClassHelper
-	{
-		bool bShouldReplace;
-		UClass* BlueprintClass;
-		UClassProperty* BPGeneratedClassProp;
-
-		FReplaceBlueprintWithClassHelper() : bShouldReplace(false), BlueprintClass(NULL), BPGeneratedClassProp(NULL)
-		{
-			GConfig->GetBool(TEXT("EditoronlyBP"), TEXT("bReplaceBlueprintWithClass"), bShouldReplace, GEditorIni);
-			if (bShouldReplace)
-			{
-				BlueprintClass = FindObject<UClass>(NULL, TEXT("/Script/Engine.Blueprint"));
-				ensure(BlueprintClass);
-				BPGeneratedClassProp = BlueprintClass ? FindField<UClassProperty>(BlueprintClass, TEXT("GeneratedClass")) : NULL;
-				ensure(BPGeneratedClassProp);
-			}
-		}
-
-		bool CanReplace() const
-		{
-			return bShouldReplace && BlueprintClass && BPGeneratedClassProp;
-		}
-	};
-	static FReplaceBlueprintWithClassHelper Helper;
-
-	const UObject* Object = GetObjectPropertyValue(Value);
 	Super::CheckValidObject(Value);
-	const UObject* CurrentObject = GetObjectPropertyValue(Value);
-
-	if (Helper.CanReplace()
-		&& !CurrentObject
-		&& Object 
-		&& Object->IsA(Helper.BlueprintClass)
-		&& (UObject::StaticClass() == MetaClass))
-	{
-		UObject* RecoveredObject = Helper.BPGeneratedClassProp->GetPropertyValue_InContainer(Object);
-		SetObjectPropertyValue(Value, RecoveredObject);
-		UE_LOG(LogProperty, Log,
-			TEXT("Blueprint '%s' is replaced with class '%s' in property '%s'"),
-			*Object->GetFullName(),
-			*RecoveredObject->GetFullName(),
-			*GetFullName());
-	}
+	// Ugly hack to replace Blueprint references with Class references.
+	FLegacyEditorOnlyBlueprintOptions::FixupClassProperty(this, Value);
 #else	// WITH_EDITOR
 	Super::CheckValidObject(Value);
 #endif	// WITH_EDITOR
diff --git a/Engine/Source/Runtime/CoreUObject/Private/UObject/PropertyMap.cpp b/Engine/Source/Runtime/CoreUObject/Private/UObject/PropertyMap.cpp
--- a/Engine/Source/Runtime/CoreUObject/Private/UObject/PropertyMap.cpp
+++ b/Engine/Source/Runtime/CoreUObject/Private/UObject/PropertyMap.cpp
@@ -235,6 +235,7 @@ void UMapProperty::GetPreloadDependencies(TArray<UObject*>& OutDeps)
 
 void UMapProperty::SerializeItem(FArchive& Ar, void* Value, const void* Defaults) const
 {
+	// Ar related calls in this function must be mirrored in UMapProperty::ConvertFromType
 	checkSlow(KeyProp);
 	checkSlow(ValueProp);
 
@@ -287,18 +288,18 @@ void UMapProperty::SerializeItem(FArchive& Ar, void* Value, const void* Defaults
 			}
 		}
 
-		int32 Num = 0;
-		Ar << Num;
+		int32 NumEntries = 0;
+		Ar << NumEntries;
 
 		// Allocate temporary key space if we haven't allocated it already above
-		if (Num != 0 && !TempKeyStorage)
+		if (NumEntries != 0 && !TempKeyStorage)
 		{
 			TempKeyStorage = (uint8*)FMemory::Malloc(MapLayout.SetLayout.Size);
 			KeyProp->InitializeValue(TempKeyStorage);
 		}
 
 		// Read remaining items into container
-		for (; Num; --Num)
+		for (; NumEntries; --NumEntries)
 		{
 			// Read key into temporary storage
 			KeyProp->SerializeItem(Ar, TempKeyStorage);
@@ -448,6 +449,14 @@ FString UMapProperty::GetCPPType(FString* ExtendedTypeText, uint32 CPPExportFlag
 	return TEXT("TMap");
 }
 
+FString UMapProperty::GetCPPTypeForwardDeclaration() const
+{
+	checkSlow(KeyProp);
+	checkSlow(ValueProp);
+	// Generates a single ' ' when no forward declaration is needed. Purely an aesthetic concern at this time:
+	return FString::Printf( TEXT("%s %s"), *KeyProp->GetCPPTypeForwardDeclaration(), *ValueProp->GetCPPTypeForwardDeclaration());
+}
+
 FString UMapProperty::GetCPPMacroType( FString& ExtendedTypeText ) const
 {
 	checkSlow(KeyProp);
@@ -794,6 +803,182 @@ bool UMapProperty::SameType(const UProperty* Other) const
 	return Super::SameType(Other) && KeyProp && ValueProp && KeyProp->SameType(MapProp->KeyProp) && ValueProp->SameType(MapProp->ValueProp);
 }
 
+bool UMapProperty::ConvertFromType(const FPropertyTag& Tag, FArchive& Ar, uint8* Data, UStruct* DefaultsStruct, bool& bOutAdvanceProperty)
+{
+	// Ar related calls in this function must be mirrored in UMapProperty::SerializeItem
+	checkSlow(KeyProp);
+	checkSlow(ValueProp);
+
+	// Ensure that the key/value properties have been loaded before calling SerializeItem() on them
+	Ar.Preload(KeyProp);
+	Ar.Preload(ValueProp);
+
+	const auto SerializeOrConvert = [](UProperty* CurrentType, const FPropertyTag& InTag, FArchive& InAr, uint8* InData, UStruct* InDefaultsStruct) -> bool
+	{
+		bool bDummyAdvance = false;
+		if(CurrentType->GetID() == InTag.Type)
+		{
+			CurrentType->SerializeItem(InAr, InData, InDefaultsStruct);
+			return true;
+		}
+		else if( CurrentType->ConvertFromType(InTag, InAr, InData, InDefaultsStruct, bDummyAdvance) )
+		{
+			return true;
+		}
+		return false;
+	};
+
+	if (Tag.Type == NAME_MapProperty && ( (Tag.InnerType != NAME_None && Tag.InnerType != KeyProp->GetID()) || (Tag.ValueType != NAME_None && Tag.ValueType != ValueProp->GetID()) ) )
+	{
+		FScriptMapHelper MapHelper(this, ContainerPtrToValuePtr<void>(Data));
+
+		uint8* TempKeyStorage = nullptr;
+		ON_SCOPE_EXIT
+		{
+			if (TempKeyStorage)
+			{
+				KeyProp->DestroyValue(TempKeyStorage);
+				FMemory::Free(TempKeyStorage);
+			}
+		};
+
+		FPropertyTag KeyPropertyTag;
+		KeyPropertyTag.Type = Tag.InnerType;
+		KeyPropertyTag.ArrayIndex = 0;
+		
+		FPropertyTag ValuePropertyTag;
+		ValuePropertyTag.Type = Tag.ValueType;
+		ValuePropertyTag.ArrayIndex = 0;
+		
+		bool bConversionSucceeded = true;
+		
+		// When we saved this instance we wrote out any elements that were in the 'Default' instance but not in the 
+		// instance that was being written. Presumably we were constructed from our defaults and must now remove 
+		// any of the elements that were not present when we saved this Map:
+		int32 NumKeysToRemove = 0;
+		Ar << NumKeysToRemove;
+
+		if( NumKeysToRemove != 0 )
+		{
+			TempKeyStorage = (uint8*)FMemory::Malloc(MapLayout.SetLayout.Size);
+			KeyProp->InitializeValue(TempKeyStorage);
+
+			if (SerializeOrConvert( KeyProp, KeyPropertyTag, Ar, TempKeyStorage, DefaultsStruct))
+			{
+				// If the key is in the map, remove it
+				int32 Found = MapHelper.FindMapIndexWithKey(TempKeyStorage);
+				if (Found != INDEX_NONE)
+				{
+					MapHelper.RemoveAt_NeedsRehash(Found);
+				}
+
+				// things are going fine, remove the rest of the keys:
+				for(int32 I = 1; I < NumKeysToRemove; ++I)
+				{
+					verify(SerializeOrConvert( KeyProp, KeyPropertyTag, Ar, TempKeyStorage, DefaultsStruct));
+					Found = MapHelper.FindMapIndexWithKey(TempKeyStorage);
+					if (Found != INDEX_NONE)
+					{
+						MapHelper.RemoveAt_NeedsRehash(Found);
+					}
+				}
+			}
+			else
+			{
+				bConversionSucceeded = false;
+			}
+		}
+
+		int32 NumEntries = 0;
+		Ar << NumEntries;
+		
+		if( bConversionSucceeded )
+		{
+			if( NumEntries != 0 )
+			{
+				if( TempKeyStorage == nullptr )
+				{
+					TempKeyStorage = (uint8*)FMemory::Malloc(MapLayout.SetLayout.Size);
+					KeyProp->InitializeValue(TempKeyStorage);
+				}
+
+				if( SerializeOrConvert( KeyProp, KeyPropertyTag, Ar, TempKeyStorage, DefaultsStruct ) )
+				{
+					// Add a new default value if the key doesn't currently exist in the map
+					bool bKeyAlreadyPresent = true;
+					int32 NextPairIndex = MapHelper.FindMapIndexWithKey(TempKeyStorage);
+					if (NextPairIndex == INDEX_NONE)
+					{
+						bKeyAlreadyPresent = false;
+						NextPairIndex = MapHelper.AddDefaultValue_Invalid_NeedsRehash();
+					}
+					
+					uint8* NextPairPtr = MapHelper.GetPairPtrWithoutCheck(NextPairIndex);
+					// This copy is unnecessary when the key was already in the map:
+					KeyProp->CopyCompleteValue_InContainer(NextPairPtr, TempKeyStorage);
+
+					// Deserialize value
+					if( SerializeOrConvert( ValueProp, ValuePropertyTag, Ar, NextPairPtr + MapLayout.ValueOffset, DefaultsStruct ) )
+					{
+						// first entry went fine, convert the rest:
+						for(int32 I = 1; I < NumEntries; ++I)
+						{
+							verify( SerializeOrConvert( KeyProp, KeyPropertyTag, Ar, TempKeyStorage, DefaultsStruct ) );
+							NextPairIndex = MapHelper.FindMapIndexWithKey(TempKeyStorage);
+							if (NextPairIndex == INDEX_NONE)
+							{
+								NextPairIndex = MapHelper.AddDefaultValue_Invalid_NeedsRehash();
+							}
+					
+							NextPairPtr = MapHelper.GetPairPtrWithoutCheck(NextPairIndex);
+							// This copy is unnecessary when the key was already in the map:
+							KeyProp->CopyCompleteValue_InContainer(NextPairPtr, TempKeyStorage);
+							verify( SerializeOrConvert( ValueProp, ValuePropertyTag, Ar, NextPairPtr + MapLayout.ValueOffset, DefaultsStruct ) );
+						}
+					}
+					else
+					{
+						if(!bKeyAlreadyPresent)
+						{
+							MapHelper.RemoveAt_NeedsRehash(NextPairIndex);
+						}
+							
+						bConversionSucceeded = false;
+					}
+				}
+				else
+				{
+					bConversionSucceeded = false;
+				}
+				
+				MapHelper.Rehash();
+			}
+		}
+
+		// if we could not convert the property ourself, then indicate that calling code needs to advance the property
+		if(!bConversionSucceeded)
+		{
+			UE_LOG(
+				LogClass, 
+				Warning, 
+				TEXT("Map Element Type mismatch in %s of %s - Previous (%s to %s) Current (%s to %s) for package: %s"), 
+				*Tag.Name.ToString(),
+				*GetName(), 
+				*Tag.InnerType.ToString(), 
+				*Tag.ValueType.ToString(), 
+				*KeyProp->GetID().ToString(), 
+				*ValueProp->GetID().ToString(), 
+				*Ar.GetArchiveName() 
+			);
+		}
+
+		bOutAdvanceProperty = bConversionSucceeded;
+
+		return true;
+	}
+	return false;
+}
+
 /**
  * Checks to see if this property already has the supplied value as a key
  * @param	InMap			The address of the map
diff --git a/Engine/Source/Runtime/CoreUObject/Private/UObject/PropertySet.cpp b/Engine/Source/Runtime/CoreUObject/Private/UObject/PropertySet.cpp
--- a/Engine/Source/Runtime/CoreUObject/Private/UObject/PropertySet.cpp
+++ b/Engine/Source/Runtime/CoreUObject/Private/UObject/PropertySet.cpp
@@ -210,6 +210,7 @@ void USetProperty::GetPreloadDependencies(TArray<UObject*>& OutDeps)
 
 void USetProperty::SerializeItem(FArchive& Ar, void* Value, const void* Defaults) const
 {
+	// Ar related calls in this function must be mirrored in USetProperty::ConvertFromType
 	checkSlow(ElementProp);
 
 	// Ensure that the element property has been loaded before calling SerializeItem() on it
@@ -713,6 +714,138 @@ bool USetProperty::SameType(const UProperty* Other) const
 	return Super::SameType(Other) && ElementProp && ElementProp->SameType(SetProp->ElementProp);
 }
 
+bool USetProperty::ConvertFromType(const FPropertyTag& Tag, FArchive& Ar, uint8* Data, UStruct* DefaultsStruct, bool& bOutAdvanceProperty)
+{
+	// Ar related calls in this function must be mirrored in USetProperty::ConvertFromType
+	checkSlow(ElementProp);
+
+	// Ensure that the element property has been loaded before calling SerializeItem() on it
+	Ar.Preload(ElementProp);
+
+	// Ar related calls in this function must be mirrored in USetProperty::SerializeItem
+	if (Tag.Type == NAME_SetProperty && Tag.InnerType != NAME_None && Tag.InnerType != ElementProp->GetID())
+	{
+		FScriptSetHelper ScriptSetHelper(this, ContainerPtrToValuePtr<void>(Data));
+		
+		uint8* TempElementStorage = nullptr;
+		ON_SCOPE_EXIT
+		{
+			if (TempElementStorage)
+			{
+				ElementProp->DestroyValue(TempElementStorage);
+				FMemory::Free(TempElementStorage);
+			}
+		};
+		
+		FPropertyTag InnerPropertyTag;
+		InnerPropertyTag.Type = Tag.InnerType;
+		InnerPropertyTag.ArrayIndex = 0;
+		
+		bool bConversionSucceeded = true;
+		bool bDummyAdvance = false;
+
+		// When we saved this instance we wrote out any elements that were in the 'Default' instance but not in the 
+		// instance that was being written. Presumably we were constructed from our defaults and must now remove 
+		// any of the elements that were not present when we saved this Set:
+		int32 NumElementsToRemove = 0;
+		Ar << NumElementsToRemove;
+
+		if(NumElementsToRemove)
+		{
+			TempElementStorage = (uint8*)FMemory::Malloc(SetLayout.Size);
+			ElementProp->InitializeValue(TempElementStorage);
+
+			if (ElementProp->ConvertFromType(InnerPropertyTag, Ar, TempElementStorage, DefaultsStruct, bDummyAdvance))
+			{
+				int32 Found = ScriptSetHelper.FindElementIndex(TempElementStorage);
+				if (Found != INDEX_NONE)
+				{
+					ScriptSetHelper.RemoveAt_NeedsRehash(Found);
+				}
+
+				for (int32 I = 1; I < NumElementsToRemove; ++I)
+				{
+					verify(ElementProp->ConvertFromType(InnerPropertyTag, Ar, TempElementStorage, DefaultsStruct, bDummyAdvance));
+					
+					Found = ScriptSetHelper.FindElementIndex(TempElementStorage);
+					if (Found != INDEX_NONE)
+					{
+						ScriptSetHelper.RemoveAt_NeedsRehash(Found);
+					}
+				}
+			}
+			else
+			{
+				bConversionSucceeded = false;
+			}
+		}
+		
+		int32 Num = 0;
+		Ar << Num;
+
+		if(bConversionSucceeded)
+		{
+			if (Num != 0)
+			{
+				// Allocate temporary key space if we haven't allocated it already above
+				if( TempElementStorage == nullptr )
+				{
+					TempElementStorage = (uint8*)FMemory::Malloc(SetLayout.Size);
+					ElementProp->InitializeValue(TempElementStorage);
+				}
+
+				// and read the first entry, we have to check for conversion possibility again because 
+				// NumElementsToRemove may not have run (in fact, it likely did not):
+				if (ElementProp->ConvertFromType(InnerPropertyTag, Ar, TempElementStorage, DefaultsStruct, bDummyAdvance))
+				{
+					if (ScriptSetHelper.FindElementIndex(TempElementStorage) == INDEX_NONE)
+					{
+						const int32 NewElementIndex = ScriptSetHelper.AddDefaultValue_Invalid_NeedsRehash();
+						uint8* NewElementPtr = ScriptSetHelper.GetElementPtrWithoutCheck(NewElementIndex);
+
+						// Copy over deserialized key from temporary storage
+						ElementProp->CopyCompleteValue_InContainer(NewElementPtr, TempElementStorage);
+					}
+
+					// Read remaining items into container
+					for (int32 I = 1; I < Num; ++I)
+					{
+						// Read key into temporary storage
+						verify(ElementProp->ConvertFromType(InnerPropertyTag, Ar, TempElementStorage, DefaultsStruct, bDummyAdvance) );
+
+						// Add a new entry if the element doesn't currently exist in the set
+						if (ScriptSetHelper.FindElementIndex(TempElementStorage) == INDEX_NONE)
+						{
+							const int32 NewElementIndex = ScriptSetHelper.AddDefaultValue_Invalid_NeedsRehash();
+							uint8* NewElementPtr = ScriptSetHelper.GetElementPtrWithoutCheck(NewElementIndex);
+
+							// Copy over deserialized key from temporary storage
+							ElementProp->CopyCompleteValue_InContainer(NewElementPtr, TempElementStorage);
+						}
+					}
+				}
+				else
+				{
+					bConversionSucceeded = false;
+				}
+			}
+		
+			ScriptSetHelper.Rehash();
+		}
+
+		// if we could not convert the property ourself, then indicate that calling code needs to advance the property
+		if(!bConversionSucceeded)
+		{
+			UE_LOG(LogClass, Warning, TEXT("Set Element Type mismatch in %s of %s - Previous (%s) Current (%s) for package: %s"), *Tag.Name.ToString(), *GetName(), *Tag.InnerType.ToString(), *ElementProp->GetID().ToString(), *Ar.GetArchiveName() );
+		}
+		bOutAdvanceProperty = bConversionSucceeded;
+
+		return true;
+	}
+
+	return false;
+}
+
 /**
  * Checks to see if this property already has the supplied value as an element
  * @param	InSet			The address of the set
diff --git a/Engine/Source/Runtime/CoreUObject/Private/UObject/PropertyTag.cpp b/Engine/Source/Runtime/CoreUObject/Private/UObject/PropertyTag.cpp
--- a/Engine/Source/Runtime/CoreUObject/Private/UObject/PropertyTag.cpp
+++ b/Engine/Source/Runtime/CoreUObject/Private/UObject/PropertyTag.cpp
@@ -2,6 +2,7 @@
 
 #include "CoreUObjectPrivate.h"
 #include "PropertyTag.h"
+#include "UObject/BlueprintsObjectVersion.h"
 
 /*-----------------------------------------------------------------------------
 FPropertyTag
@@ -15,6 +16,7 @@ FPropertyTag::FPropertyTag()
 	, StructName(NAME_None)
 	, EnumName  (NAME_None)
 	, InnerType (NAME_None)
+	, ValueType	(NAME_None)
 	, Size      (0)
 	, ArrayIndex(INDEX_NONE)
 	, SizeOffset(INDEX_NONE)
@@ -26,9 +28,10 @@ FPropertyTag::FPropertyTag( FArchive& InSaveAr, UProperty* Property, int32 InInd
 	, BoolVal   (0)
 	, Name      (Property->GetFName())
 	, StructName(NAME_None)
-	, EnumName  (NAME_None)
-	, InnerType (NAME_None)
-	, Size      (0)
+	, EnumName	(NAME_None)
+	, InnerType	(NAME_None)
+	, ValueType	(NAME_None)
+	, Size		(0)
 	, ArrayIndex(InIndex)
 	, SizeOffset(INDEX_NONE)
 	, HasPropertyGuid(0)
@@ -52,6 +55,15 @@ FPropertyTag::FPropertyTag( FArchive& InSaveAr, UProperty* Property, int32 InInd
 		{
 			InnerType = ArrayProp->Inner->GetID();
 		}
+		else if (USetProperty* SetProp = Cast<USetProperty>(Property))
+		{
+			InnerType = SetProp->ElementProp->GetID();
+		}
+		else if (UMapProperty* MapProp = Cast<UMapProperty>(Property))
+		{
+			InnerType = MapProp->KeyProp->GetID();
+			ValueType = MapProp->ValueProp->GetID();
+		}
 		else if (UBoolProperty* Bool = Cast<UBoolProperty>(Property))
 		{
 			BoolVal = Bool->GetPropertyValue(Value);
@@ -117,6 +129,20 @@ FArchive& operator<<( FArchive& Ar, FPropertyTag& Tag )
 			Ar << Tag.InnerType;
 		}
 	}
+
+	if (Ar.UE4Ver() >= VER_UE4_PROPERTY_TAG_SET_MAP_SUPPORT)
+	{
+		if (Tag.Type == NAME_SetProperty)
+		{
+			Ar << Tag.InnerType;
+		}
+		else if (Tag.Type == NAME_MapProperty)
+		{
+			Ar << Tag.InnerType;
+			Ar << Tag.ValueType;
+		}
+	}
+
 	// Property tags to handle renamed blueprint properties effectively.
 	if (Ar.UE4Ver() >= VER_UE4_PROPERTY_GUID_IN_PROPERTY_TAG)
 	{
diff --git a/Engine/Source/Runtime/CoreUObject/Private/UObject/SavePackage.cpp b/Engine/Source/Runtime/CoreUObject/Private/UObject/SavePackage.cpp
--- a/Engine/Source/Runtime/CoreUObject/Private/UObject/SavePackage.cpp
+++ b/Engine/Source/Runtime/CoreUObject/Private/UObject/SavePackage.cpp
@@ -10,6 +10,7 @@
 #include "DebugSerializationFlags.h"
 #include "UObject/GCScopeLock.h"
 #include "CookStats.h"
+#include "BlueprintSupport.h" // for FLegacyEditorOnlyBlueprintOptions::DetectInvalidBlueprintExport()
 
 DEFINE_LOG_CATEGORY_STATIC(LogSavePackage, Log, All);
 
@@ -4859,6 +4860,15 @@ ESavePackageResult UPackage::Save(UPackage* InOuter, UObject* Base, EObjectFlags
 							// this export's Outer is the LinkerRoot for this package
 							Export.OuterIndex = FPackageIndex();
 						}
+
+						// exporting UBlueprint (editor-only) objects in the cook is for some reason 
+						// incompatible with the EDL (but could occur thanks to legacy settings) - here 
+						// we catch that and warn that these settings are now deprecated (this'll 
+						// trigger a fatal error if currently using the EDL)
+						//
+						// NOTE: this is temporary, as these legacy settings are now deprecated, and
+						//       should be removed soon after
+						FLegacyEditorOnlyBlueprintOptions::DetectInvalidBlueprintExport(Linker, i);
 					}
 				}
 
diff --git a/Engine/Source/Runtime/CoreUObject/Public/Blueprint/BlueprintSupport.h b/Engine/Source/Runtime/CoreUObject/Public/Blueprint/BlueprintSupport.h
--- a/Engine/Source/Runtime/CoreUObject/Public/Blueprint/BlueprintSupport.h
+++ b/Engine/Source/Runtime/CoreUObject/Public/Blueprint/BlueprintSupport.h
@@ -102,6 +102,16 @@ enum class EReplacementResult
 	ReplaceCompletely
 };
 
+struct COREUOBJECT_API FScopedPlaceholderRawContainerTracker
+{
+public:
+	FScopedPlaceholderRawContainerTracker(void* InData);
+	~FScopedPlaceholderRawContainerTracker();
+
+private:
+	void* Data;
+};
+
 /**
  * Interface needed by CoreUObject to the BlueprintNativeCodeGen logic. Used by cooker to convert assets 
  * to native code.
@@ -253,4 +263,24 @@ public:
 	void GetAssets(FName PackageName, TArray<FBlueprintDependencyData>& OutDependencies) const;
 
 	static void FillUsedAssetsInDynamicClass(UDynamicClass* DynamicClass, GetDependenciesNamesFunc GetUsedAssets);
-};
\ No newline at end of file
+};
+
+/** 
+ * Temporary util struct that consolidates the now deprecated [EditoronlyBP] 
+ * settings. Aimed at providing warnings for users who are mistakenly using or 
+ * relying on them (and centralizing their usage so it is easy to strip in later 
+ * versions).
+ *
+ * Here in the CoreUObject module because UClassProperty::CheckValidObject() and 
+ * UPackage::Save() use some of these settings.
+ */
+struct COREUOBJECT_API FLegacyEditorOnlyBlueprintOptions
+{
+	static FString GetDefaultEditorConfig();
+	static bool IncludeUBlueprintObjsInCookedBuilds();
+	static bool AllowLegacyBlueprintPinMatchesWithClass();
+	static bool FixupLegacyBlueprintReferences();
+	static bool DetectInvalidBlueprintExport(const FLinkerSave* Linker, const int32 ExportIndex);
+	static bool FixupClassProperty(const UClassProperty* Property, void* Value);
+	static bool IsTypeProhibited(const UClass* VarType);
+};
diff --git a/Engine/Source/Runtime/CoreUObject/Public/Serialization/ObjectReader.h b/Engine/Source/Runtime/CoreUObject/Public/Serialization/ObjectReader.h
--- a/Engine/Source/Runtime/CoreUObject/Public/Serialization/ObjectReader.h
+++ b/Engine/Source/Runtime/CoreUObject/Public/Serialization/ObjectReader.h
@@ -7,7 +7,7 @@
 struct FObjectInstancingGraph;
 
 /**
- * UObject Memory Reader Archive.
+ * UObject Memory Reader Archive. Reads from InBytes, writes to Obj.
  */
 class FObjectReader : public FMemoryArchive
 {
diff --git a/Engine/Source/Runtime/CoreUObject/Public/UObject/NoExportTypes.h b/Engine/Source/Runtime/CoreUObject/Public/UObject/NoExportTypes.h
--- a/Engine/Source/Runtime/CoreUObject/Public/UObject/NoExportTypes.h
+++ b/Engine/Source/Runtime/CoreUObject/Public/UObject/NoExportTypes.h
@@ -487,7 +487,7 @@ struct FIntPoint
  * An integer vector in 3D space.
  * The full C++ class is located here: Engine\Source\Runtime\Core\Public\Math\IntVector.h
  */
-USTRUCT(immutable, noexport)
+USTRUCT(immutable, noexport, BlueprintType)
 struct FIntVector
 {
 	UPROPERTY(EditAnywhere, BlueprintReadWrite, Category=IntVector, SaveGame)
diff --git a/Engine/Source/Runtime/CoreUObject/Public/UObject/PropertyTag.h b/Engine/Source/Runtime/CoreUObject/Public/UObject/PropertyTag.h
--- a/Engine/Source/Runtime/CoreUObject/Public/UObject/PropertyTag.h
+++ b/Engine/Source/Runtime/CoreUObject/Public/UObject/PropertyTag.h
@@ -21,10 +21,11 @@ struct FPropertyTag
 	FName	Name;		// Name of property.
 	FName	StructName;	// Struct name if UStructProperty.
 	FName	EnumName;	// Enum name if UByteProperty
-	FName InnerType; // Inner type if UArrayProperty
-	int32		Size;       // Property size.
-	int32		ArrayIndex;	// Index if an array; else 0.
-	int64		SizeOffset;	// location in stream of tag size member
+	FName	InnerType;	// Inner type if UArrayProperty, USetProperty, or UMapProperty
+	FName	ValueType;	// Value type if UMapPropery
+	int32	Size;       // Property size.
+	int32	ArrayIndex;	// Index if an array; else 0.
+	int64	SizeOffset;	// location in stream of tag size member
 	FGuid	StructGuid;
 	uint8	HasPropertyGuid;
 	FGuid	PropertyGuid;
diff --git a/Engine/Source/Runtime/CoreUObject/Public/UObject/UnrealType.h b/Engine/Source/Runtime/CoreUObject/Public/UObject/UnrealType.h
--- a/Engine/Source/Runtime/CoreUObject/Public/UObject/UnrealType.h
+++ b/Engine/Source/Runtime/CoreUObject/Public/UObject/UnrealType.h
@@ -201,7 +201,9 @@ public:
 	* @param	DefaultsStruct 
 	* @param	bOutAdvanceProperty whether the property should be advanced and continue to next property or not
 	*
-	* @return	true if the function has handled the tag
+	* @return	true if the function has handled the tag. Returning true but assigning bOutAdvanceProperty false 
+	*			indicates to calling code that it needs to manually advance the property. bOutAdvanceProperty true
+	*			indicates that ConvertFromType itself has handled this property tag.
 	*/
 	virtual bool ConvertFromType(const FPropertyTag& Tag, FArchive& Ar, uint8* Data, UStruct* DefaultsStruct, bool& bOutAdvanceProperty);
 
@@ -2490,6 +2492,7 @@ public:
 	// UProperty interface
 	virtual FString GetCPPMacroType(FString& ExtendedTypeText) const  override;
 	virtual FString GetCPPType(FString* ExtendedTypeText, uint32 CPPExportFlags) const override;
+	virtual FString GetCPPTypeForwardDeclaration() const override;
 	virtual void LinkInternal(FArchive& Ar) override;
 	virtual bool Identical(const void* A, const void* B, uint32 PortFlags) const override;
 	virtual void SerializeItem(FArchive& Ar, void* Value, void const* Defaults) const override;
@@ -2505,6 +2508,7 @@ public:
 	virtual bool ContainsWeakObjectReference() const override;
 	virtual void EmitReferenceInfo(UClass& OwnerClass, int32 BaseOffset, TArray<const UStructProperty*>& EncounteredStructProps) override;
 	virtual bool SameType(const UProperty* Other) const override;
+	virtual bool ConvertFromType(const FPropertyTag& Tag, FArchive& Ar, uint8* Data, UStruct* DefaultsStruct, bool& bOutAdvanceProperty) override;
 	// End of UProperty interface
 
 	bool HasKey(void* InMap, void* InBaseAddress, const FString& InValue) const;
@@ -2556,6 +2560,7 @@ public:
 	virtual bool ContainsWeakObjectReference() const override;
 	virtual void EmitReferenceInfo(UClass& OwnerClass, int32 BaseOffset, TArray<const UStructProperty*>& EncounteredStructProps) override;
 	virtual bool SameType(const UProperty* Other) const override;
+	virtual bool ConvertFromType(const FPropertyTag& Tag, FArchive& Ar, uint8* Data, UStruct* DefaultsStruct, bool& bOutAdvanceProperty) override;
 	// End of UProperty interface
 
 	bool HasElement(void* InSet, void* InBaseAddress, const FString& InValue) const;
@@ -3576,15 +3581,16 @@ public:
 		);
 	}
 
-	/** Remoes the elemnt from the set */
-	void RemoveElement(const void* ElementToRemove)
+	/** Removes the element from the set */
+	bool RemoveElement(const void* ElementToRemove)
 	{
 		UProperty* LocalElementPropForCapture = ElementProp;
-		Set->Remove(
+		return Set->Remove(
 			ElementToRemove, 
 			SetLayout, 
 			[LocalElementPropForCapture](const void* Element) { return LocalElementPropForCapture->GetValueTypeHash(Element); },
-			[LocalElementPropForCapture](const void* A, const void* B) { return LocalElementPropForCapture->Identical(A, B); }
+			[LocalElementPropForCapture](const void* A, const void* B) { return LocalElementPropForCapture->Identical(A, B); },
+			[LocalElementPropForCapture](void* ElementToDestroy ) {LocalElementPropForCapture->DestroyValue(ElementToDestroy); }
 		);
 	}
 
diff --git a/Engine/Source/Runtime/Engine/Classes/Engine/Blueprint.h b/Engine/Source/Runtime/Engine/Classes/Engine/Blueprint.h
--- a/Engine/Source/Runtime/Engine/Classes/Engine/Blueprint.h
+++ b/Engine/Source/Runtime/Engine/Classes/Engine/Blueprint.h
@@ -344,6 +344,10 @@ class ENGINE_API UBlueprint : public UBlueprintCore
 	/** Additional HideCategories. These are added to HideCategories from parent. */
 	UPROPERTY(EditAnywhere, Category=BlueprintOptions)
 	TArray<FString> HideCategories;
+	 
+	/** When nativization is enabled, and "NativizeOnlySelectedBlueprints" is set true, then this asset will be nativized. All super classes must be also nativized. */
+	UPROPERTY(EditAnywhere, Category=Experimental)
+	bool bNativize;
 
 	/** TRUE to show a warning when attempting to start in PIE and there is a compiler error on this Blueprint */
 	UPROPERTY(transient)
@@ -463,6 +467,10 @@ class ENGINE_API UBlueprint : public UBlueprintCore
 	UPROPERTY()
 	TArray<class UEdGraphPin_Deprecated*> DeprecatedPinWatches;
 
+	/** Index map for component template names */
+	UPROPERTY()
+	TMap<FName, int32> ComponentTemplateNameIndex;
+
 	/** Maps old to new component template names */
 	UPROPERTY(transient)
 	TMap<FName, FName> OldToNewComponentTemplateNames;
@@ -659,6 +667,9 @@ public:
 	virtual bool NeedsLoadForClient() const override;
 	virtual bool NeedsLoadForServer() const override;
 	virtual bool NeedsLoadForEditorGame() const override;
+#if WITH_EDITOR
+	virtual bool CanEditChange(const UProperty* InProperty) const override;
+#endif // WITH_EDITOR
 	//~ End UObject Interface
 
 	/** Get the Blueprint object that generated the supplied class */
diff --git a/Engine/Source/Runtime/Engine/Classes/Engine/DataTable.h b/Engine/Source/Runtime/Engine/Classes/Engine/DataTable.h
--- a/Engine/Source/Runtime/Engine/Classes/Engine/DataTable.h
+++ b/Engine/Source/Runtime/Engine/Classes/Engine/DataTable.h
@@ -42,7 +42,7 @@ struct FTableRowBase
 /**
  * Imported spreadsheet table.
  */
-UCLASS(MinimalAPI)
+UCLASS(MinimalAPI, BlueprintType)
 class UDataTable
 	: public UObject
 {
diff --git a/Engine/Source/Runtime/Engine/Classes/GameFramework/ProjectileMovementComponent.h b/Engine/Source/Runtime/Engine/Classes/GameFramework/ProjectileMovementComponent.h
--- a/Engine/Source/Runtime/Engine/Classes/GameFramework/ProjectileMovementComponent.h
+++ b/Engine/Source/Runtime/Engine/Classes/GameFramework/ProjectileMovementComponent.h
@@ -293,9 +293,11 @@ protected:
 	/** Allow the projectile to track towards its homing target. */
 	virtual FVector ComputeHomingAcceleration(const FVector& InVelocity, float DeltaTime) const;
 
+public:
 	/** Compute gravity effect given current physics volume, projectile gravity scale, etc. */
 	virtual float GetGravityZ() const override;
 
+protected:
 	DEPRECATED(4.9, "GetEffectiveGravityZ() is deprecated, use GetGravityZ() instead, or UMovementComponent::GetGravityZ() if you want only the unmodified gravity.")
 	float GetEffectiveGravityZ() const;
 
diff --git a/Engine/Source/Runtime/Engine/Classes/Kismet/BlueprintSetLibrary.h b/Engine/Source/Runtime/Engine/Classes/Kismet/BlueprintSetLibrary.h
--- a/Engine/Source/Runtime/Engine/Classes/Kismet/BlueprintSetLibrary.h
+++ b/Engine/Source/Runtime/Engine/Classes/Kismet/BlueprintSetLibrary.h
@@ -7,78 +7,81 @@
 UCLASS(Experimental)
 class ENGINE_API UBlueprintSetLibrary : public UBlueprintFunctionLibrary
 {
-	GENERATED_UCLASS_BODY()
+	GENERATED_BODY()
 
+public:
 	/** 
-	 *Add item to set
+	 * Adds item to set
 	 *
-	 *@param	TargetSet		The set to add item to
-	 *@param	NewItem			The item to add to the set
-	*/
+	 * @param	TargetSet		The set to add item to
+	 * @param	NewItem			The item to add to the set
+	 * @return	True if NewItem was added to the set (False indicates an equivalent item was present)
+	 * 
+	 */
 	UFUNCTION(BlueprintCallable, CustomThunk, meta=(DisplayName = "Add", CompactNodeTitle = "ADD", SetParam = "TargetSet|NewItem", AutoCreateRefTerm = "NewItem"), Category="Utilities|Set")
-	static void Set_Add(const TSet<int32>& TargetSet, const int32& NewItem);
+	static bool Set_Add(const TSet<int32>& TargetSet, const int32& NewItem);
 	
 	/**
 	 * Adds all elements from an Array to a Set
 	 *
-	 *@param	TargetSet		The set to search for the item
-	 *@param	NewItems		The items to add to the set
-	*/
+	 * @param	TargetSet		The set to search for the item
+	 * @param	NewItems		The items to add to the set
+	 */
 	UFUNCTION(BlueprintCallable, CustomThunk, meta=(DisplayName = "Add Items", CompactNodeTitle = "ADD ITEMS", SetParam = "TargetSet|NewItems", AutoCreateRefTerm = "NewItems"), Category="Utilities|Set")
 	static void Set_AddItems(const TSet<int32>& TargetSet, const TArray<int32>& NewItems);
 
 	/**
-	 *Remove item from set.
+	 * Remove item from set.
 	 *
-	 *@param	TargetSet		The set to remove from
-	 *@param	Item			The item to remove from the set
-	 *@return	True if an item was removed
-	*/
+	 * @param	TargetSet		The set to remove from
+	 * @param	Item			The item to remove from the set
+	 * @return	True if an item was removed (False indicates no equivalent item was present)
+	 */
 	UFUNCTION(BlueprintCallable, CustomThunk, meta=(DisplayName = "Remove", CompactNodeTitle = "REMOVE", SetParam = "TargetSet|Item", AutoCreateRefTerm = "Item"), Category="Utilities|Set")
 	static bool Set_Remove(const TSet<int32>& TargetSet, const int32& Item);
 	
 	/**
-	 *Removes all elements in an Array from a set.
+	 * Removes all elements in an Array from a set.
 	 *
-	 *@param	TargetSet		The set to remove from
-	 *@param	Items			The items to remove from the set
-	*/
+	 * @param	TargetSet		The set to remove from
+	 * @param	Items			The items to remove from the set
+	 */
 	UFUNCTION(BlueprintCallable, CustomThunk, meta=(DisplayName = "Remove Items", CompactNodeTitle = "REMOVE ITEMS", SetParam = "TargetSet|Items", AutoCreateRefTerm = "Items"), Category="Utilities|Set")
 	static void Set_RemoveItems(const TSet<int32>& TargetSet, const  TArray<int32>& Items);
 
 	/**
-	* Returns an Array, Result, with the entries of a Set.
-	*
-	* @param		A		Set
-	* @param		Result	Array
-	*/
+	 * Outputs an Array containing copies of the entries of a Set.
+	 *
+	 * @param		A		Set
+	 * @param		Result	Array
+	 */
 	UFUNCTION(BlueprintCallable, CustomThunk, meta = (DisplayName = "To Array", CompactNodeTitle = "TO ARRAY", SetParam = "A|Result"), Category = "Utilities|Set")
 	static void Set_ToArray(const TSet<int32>& A, TArray<int32>& Result);
 
 	/**
-	 *Clear a set, removes all content
+	 * Clear a set, removes all content
 	 *
-	 *@param	TargetSet		The set to clear
-	*/
+	 * @param	TargetSet		The set to clear
+	 */
 	UFUNCTION(BlueprintCallable, CustomThunk, meta=(DisplayName = "Clear", CompactNodeTitle = "CLEAR", Keywords = "empty", SetParam = "TargetSet"), Category="Utilities|Set")
 	static void Set_Clear(const TSet<int32>& TargetSet);
 
 	/**
 	 * Get the number of items in a set
 	 *
-	 *@param	TargetSet		The set to get the length of
-	 *@return	The length of the set
-	*/
+	 * @param	TargetSet		The set to get the length of
+	 * @return	The length of the set
+	 */
 	UFUNCTION(BlueprintPure, CustomThunk, meta=(DisplayName = "Length", CompactNodeTitle = "LENGTH", SetParam = "TargetSet", Keywords = "num size count"), Category="Utilities|Set")
 	static int32 Set_Length(const TSet<int32>& TargetSet);
 
 	/**
-	 *Returns true if the set contains the given item
+	 * Returns true if the set contains the given item
 	 *
-	 *@param	TargetSet		The set to search for the item
-	 *@param	ItemToFind		The item to look for
-	 *@return	True if the item was found within the set
-	*/
+	 * @param	TargetSet		The set to search for the item
+	 * @param	ItemToFind		The item to look for
+	 * @return	True if the item was found within the set
+	 */
 	UFUNCTION(BlueprintPure, CustomThunk, meta=(DisplayName = "Contains Item", CompactNodeTitle = "CONTAINS", SetParam = "TargetSet|ItemToFind", AutoCreateRefTerm = "ItemToFind"), Category="Utilities|Set")
 	static bool Set_Contains(const TSet<int32>& TargetSet, const int32& ItemToFind);
 
@@ -87,9 +90,9 @@ class ENGINE_API UBlueprintSetLibrary : public UBlueprintFunctionLibrary
 	 * all elements that are in both Set A and Set B. To intersect with the empty set use
 	 * Clear.
 	 *
-	 *@param		A		One set to intersect
-	 *@param		B		Another set to intersect
-	 *@param		Result	Set to store results in
+	 * @param		A		One set to intersect
+	 * @param		B		Another set to intersect
+	 * @param		Result	Set to store results in
 	 */
 	UFUNCTION(BlueprintCallable, CustomThunk, meta=(DisplayName = "Intersection", CompactNodeTitle = "INTERSECTION", SetParam = "A,B,Result"), Category="Utilities|Set")
 	static void Set_Intersection(const TSet<int32>& A, const TSet<int32>& B, TSet<int32>& Result );
@@ -99,9 +102,9 @@ class ENGINE_API UBlueprintSetLibrary : public UBlueprintFunctionLibrary
 	 * all elements that are in Set A and in addition all elements in Set B. Note that 
 	 * a Set is a collection of unique elements, so duplicates will be eliminated.
 	 *
-	 *@param		A		One set to union
-	 *@param		B		Another set to union
-	 *@param		Result	Set to store results in
+	 * @param		A		One set to union
+	 * @param		B		Another set to union
+	 * @param		Result	Set to store results in
 	 */
 	UFUNCTION(BlueprintCallable, CustomThunk, meta = (DisplayName = "Union", CompactNodeTitle = "UNION", SetParam = "A,B,Result"), Category = "Utilities|Set")
 	static void Set_Union(const TSet<int32>& A, const TSet<int32>& B, TSet<int32>& Result );
@@ -112,9 +115,9 @@ class ENGINE_API UBlueprintSetLibrary : public UBlueprintFunctionLibrary
 	 * difference between two sets  is not commutative. The Set whose elements you wish to 
 	 * preserve should be the first (top) parameter. Also called the relative complement.
 	 *
-	 *@param		A		Starting set
-	 *@param		B		Set of elements to remove from set A
-	 *@param		Result	Set containing all elements in A that are not found in B
+	 * @param		A		Starting set
+	 * @param		B		Set of elements to remove from set A
+	 * @param		Result	Set containing all elements in A that are not found in B
 	 */
 	UFUNCTION(BlueprintCallable, CustomThunk, meta = (DisplayName = "Difference", CompactNodeTitle = "DIFFERENCE", SetParam = "A,B,Result"), Category = "Utilities|Set")
 	static void Set_Difference(const TSet<int32>& A, const TSet<int32>& B, TSet<int32>& Result );
@@ -144,7 +147,7 @@ class ENGINE_API UBlueprintSetLibrary : public UBlueprintFunctionLibrary
 		P_FINISH;
 
 		P_NATIVE_BEGIN;
-		GenericSet_Add(SetAddr, SetProperty, ItemPtr);
+		*(bool*)RESULT_PARAM = GenericSet_Add(SetAddr, SetProperty, ItemPtr);
 		P_NATIVE_END;
 
 		ElementProp->DestroyValue(StorageSpace);
@@ -204,7 +207,7 @@ class ENGINE_API UBlueprintSetLibrary : public UBlueprintFunctionLibrary
 		P_FINISH;
 
 		P_NATIVE_BEGIN;
-		GenericSet_Remove(SetAddr, SetProperty, ItemPtr);
+		*(bool*)RESULT_PARAM = GenericSet_Remove(SetAddr, SetProperty, ItemPtr);
 		P_NATIVE_END;
 
 		ElementProp->DestroyValue(StorageSpace);
@@ -458,9 +461,9 @@ class ENGINE_API UBlueprintSetLibrary : public UBlueprintFunctionLibrary
 		P_NATIVE_END;
 	}
 
-	static void GenericSet_Add(const void* TargetSet, const USetProperty* SetProperty, const void* ItemPtr);
+	static bool GenericSet_Add(const void* TargetSet, const USetProperty* SetProperty, const void* ItemPtr);
 	static void GenericSet_AddItems(const void* TargetSet, const USetProperty* SetProperty, const void* TargetArray, const UArrayProperty* ArrayProperty);
-	static void GenericSet_Remove(const void* TargetSet, const USetProperty* SetProperty, const void* ItemPtr);
+	static bool GenericSet_Remove(const void* TargetSet, const USetProperty* SetProperty, const void* ItemPtr);
 	static void GenericSet_RemoveItems(const void* TargetSet, const USetProperty* SetProperty, const void* TargetArray, const UArrayProperty* ArrayProperty);
 	static void GenericSet_ToArray(const void* TargetSet, const USetProperty* SetProperty, void* TargetArray, const UArrayProperty* ArrayProperty);
 	static void GenericSet_Clear(const void* TargetSet, const USetProperty* SetProperty);
diff --git a/Engine/Source/Runtime/Engine/Classes/Kismet/GameplayStatics.h b/Engine/Source/Runtime/Engine/Classes/Kismet/GameplayStatics.h
--- a/Engine/Source/Runtime/Engine/Classes/Kismet/GameplayStatics.h
+++ b/Engine/Source/Runtime/Engine/Classes/Kismet/GameplayStatics.h
@@ -719,7 +719,7 @@ class ENGINE_API UGameplayStatics : public UBlueprintFunctionLibrary
 	 *	@param	SaveGameClass	Class of SaveGame to create
 	 *	@return					New SaveGame object to write data to
 	 */
-	UFUNCTION(BlueprintCallable, Category="Game")
+	UFUNCTION(BlueprintCallable, Category="Game", meta=(DeterminesOutputType="SaveGameClass"))
 	static USaveGame* CreateSaveGameObject(TSubclassOf<USaveGame> SaveGameClass);
 
 	/** 
diff --git a/Engine/Source/Runtime/Engine/Classes/Kismet/KismetMathLibrary.h b/Engine/Source/Runtime/Engine/Classes/Kismet/KismetMathLibrary.h
--- a/Engine/Source/Runtime/Engine/Classes/Kismet/KismetMathLibrary.h
+++ b/Engine/Source/Runtime/Engine/Classes/Kismet/KismetMathLibrary.h
@@ -620,6 +620,10 @@ class ENGINE_API UKismetMathLibrary : public UBlueprintFunctionLibrary
 	UFUNCTION(BlueprintPure, meta=(DisplayName = "Truncate", BlueprintAutocast), Category="Math|Float")
 	static int32 FTrunc(float A);
 	
+	/* Rounds A to an integer with truncation towards zero for each element in a vector.  (e.g. -1.7 truncated to -1, 2.8 truncated to 2) */
+	UFUNCTION(BlueprintPure, meta = (DisplayName = "Truncate Vector", BlueprintAutocast), Category = "Math|Float")
+	static FIntVector FTruncVector(const FVector& InVector);
+
 	/* Rounds A to the smallest following integer */
 	UFUNCTION(BlueprintPure, Category="Math|Float")
 	static int32 FCeil(float A);
@@ -1339,6 +1343,10 @@ class ENGINE_API UKismetMathLibrary : public UBlueprintFunctionLibrary
 	UFUNCTION(BlueprintPure, meta=(DisplayName = "ToByte (int)", CompactNodeTitle = "->", Keywords="cast convert", BlueprintAutocast), Category="Math|Conversions")
 	static uint8 Conv_IntToByte(int32 InInt);
 
+	/** Converts an integer to an IntVector*/
+	UFUNCTION(BlueprintPure, meta = (DisplayName = "ToIntVector (int)", CompactNodeTitle = "->", Keywords = "cast convert", BlueprintAutocast), Category = "Math|Conversions")
+	static FIntVector Conv_IntToIntVector(int32 InInt);
+
 	/** Converts a int to a bool*/
 	UFUNCTION(BlueprintPure, meta=(DisplayName = "ToBool (int)", CompactNodeTitle = "->", Keywords="cast convert", BlueprintAutocast), Category="Math|Conversions")
 	static bool Conv_IntToBool(int32 InInt);
@@ -1387,6 +1395,10 @@ class ENGINE_API UKismetMathLibrary : public UBlueprintFunctionLibrary
 	UFUNCTION(BlueprintPure, meta=(DisplayName = "ToVector (Vector2D)", CompactNodeTitle = "->", Keywords="cast convert", BlueprintAutocast), Category="Math|Conversions")
 	static FVector Conv_Vector2DToVector(FVector2D InVector2D, float Z = 0);
 
+	/** Convert an IntVector to a vector */
+	UFUNCTION(BlueprintPure, meta = (DisplayName = "ToVector (IntVector)", CompactNodeTitle = "->", Keywords = "cast convert", BlueprintAutocast), Category = "Math|Conversions")
+	static FVector Conv_IntVectorToVector(const FIntVector& InIntVector);
+
 	/** Convert a float into a vector, where each element is that float */
 	UFUNCTION(BlueprintPure, meta=(DisplayName = "ToVector (float)", CompactNodeTitle = "->", Keywords="cast convert", BlueprintAutocast), Category="Math|Conversions")
 	static FVector Conv_FloatToVector(float InFloat);
diff --git a/Engine/Source/Runtime/Engine/Classes/Kismet/KismetStringLibrary.h b/Engine/Source/Runtime/Engine/Classes/Kismet/KismetStringLibrary.h
--- a/Engine/Source/Runtime/Engine/Classes/Kismet/KismetStringLibrary.h
+++ b/Engine/Source/Runtime/Engine/Classes/Kismet/KismetStringLibrary.h
@@ -29,6 +29,10 @@ class ENGINE_API UKismetStringLibrary : public UBlueprintFunctionLibrary
 	UFUNCTION(BlueprintPure, meta=(DisplayName = "ToString (Vector)", CompactNodeTitle = "->", BlueprintAutocast), Category="Utilities|String")
 	static FString Conv_VectorToString(FVector InVec);
 
+	/** Converts an IntVector value to a string, in the form 'X= Y= Z=' */
+	UFUNCTION(BlueprintPure, meta=(DisplayName = "ToString (IntVector)", CompactNodeTitle = "->", BlueprintAutocast), Category="Utilities|String")
+	static FString Conv_IntVectorToString(FIntVector InIntVec);
+
 	/** Converts a vector2d value to a string, in the form 'X= Y=' */
 	UFUNCTION(BlueprintPure, meta=(DisplayName = "ToString (vector2d)", CompactNodeTitle = "->", BlueprintAutocast), Category="Utilities|String")
 	static FString Conv_Vector2dToString(FVector2D InVec);
@@ -125,6 +129,16 @@ class ENGINE_API UKismetStringLibrary : public UBlueprintFunctionLibrary
 	UFUNCTION(BlueprintPure, meta=(DisplayName = "BuildString (vector)"), Category="Utilities|String")
 	static FString BuildString_Vector(const FString& AppendTo, const FString& Prefix, FVector InVector, const FString& Suffix);
 
+	/**
+	 * Converts an IntVector->string, creating a new string in the form AppendTo+Prefix+InIntVector+Suffix
+	 * @param AppendTo - An existing string to use as the start of the conversion string
+	 * @param Prefix - A string to use as a prefix, after the AppendTo string
+	 * @param InIntVector - The intVector value to convert. Uses the standard FVector::ToString conversion
+	 * @param Suffix - A suffix to append to the end of the conversion string
+	 * @return A new string built from the passed parameters
+	 */
+	UFUNCTION(BlueprintPure, meta = (DisplayName = "BuildString (IntVector)"), Category = "Utilities|String")
+	static FString BuildString_IntVector(const FString& AppendTo, const FString& Prefix, FIntVector InIntVector, const FString& Suffix);
 	/** 
 	 * Converts a vector2d->string, creating a new string in the form AppendTo+Prefix+InVector2d+Suffix
 	 * @param AppendTo - An existing string to use as the start of the conversion string
diff --git a/Engine/Source/Runtime/Engine/Private/Blueprint.cpp b/Engine/Source/Runtime/Engine/Private/Blueprint.cpp
--- a/Engine/Source/Runtime/Engine/Private/Blueprint.cpp
+++ b/Engine/Source/Runtime/Engine/Private/Blueprint.cpp
@@ -25,6 +25,7 @@
 #endif
 #include "Components/TimelineComponent.h"
 #include "Engine/InheritableComponentHandler.h"
+#include "BlueprintSupport.h" // for FLegacyEditorOnlyBlueprintOptions::IncludeUBlueprintObjsInCookedBuilds()
 
 DEFINE_LOG_CATEGORY(LogBlueprint);
 
@@ -470,7 +471,7 @@ bool UBlueprint::Rename( const TCHAR* InName, UObject* NewOuter, ERenameFlags Fl
 	bool bSuccess = Super::Rename( InName, NewOuter, Flags );
 
 	// Finally, do a compile, but only if the new name differs from before
-	if(bSuccess && !(Flags & REN_Test) && !(Flags & REN_DoNotDirty) && InName != OldName)
+	if(bSuccess && !(Flags & REN_Test) && !(Flags & REN_DoNotDirty) && InName && InName != OldName)
 	{
 		// Gather all blueprints that currently depend on this one.
 		TArray<UBlueprint*> Dependents;
@@ -713,7 +714,7 @@ void UBlueprint::SetObjectBeingDebugged(UObject* NewObject)
 				NewObject->IsA(this->GeneratedClass), 
 				TEXT("Type mismatch: Expected %s, Found %s"), 
 				this->GeneratedClass ? *(this->GeneratedClass->GetName()) : TEXT("NULL"), 
-				NewObject->GetClass() ? *(this->GetClass()->GetName()) : TEXT("NULL")))
+				NewObject->GetClass() ? *(NewObject->GetClass()->GetName()) : TEXT("NULL")))
 		{
 			NewObject = NULL;
 		}
@@ -1229,26 +1230,14 @@ FString UBlueprint::GetDesc(void)
 	return ResultString;
 }
 
-struct FDontLoadBlueprintOutsideEditorHelper
-{
-	bool bDontLoad;
-
-	FDontLoadBlueprintOutsideEditorHelper() : bDontLoad(false)
-	{
-		GConfig->GetBool(TEXT("EditoronlyBP"), TEXT("bDontLoadBlueprintOutsideEditor"), bDontLoad, GEditorIni);
-	}
-};
-
 bool UBlueprint::NeedsLoadForClient() const
 {
-	static const FDontLoadBlueprintOutsideEditorHelper Helper;
-	return !Helper.bDontLoad;
+	return FLegacyEditorOnlyBlueprintOptions::IncludeUBlueprintObjsInCookedBuilds();
 }
 
 bool UBlueprint::NeedsLoadForServer() const
 {
-	static const FDontLoadBlueprintOutsideEditorHelper Helper;
-	return !Helper.bDontLoad;
+	return FLegacyEditorOnlyBlueprintOptions::IncludeUBlueprintObjsInCookedBuilds();
 }
 
 bool UBlueprint::NeedsLoadForEditorGame() const
@@ -1256,6 +1245,22 @@ bool UBlueprint::NeedsLoadForEditorGame() const
 	return true;
 }
 
+#if WITH_EDITOR
+bool UBlueprint::CanEditChange(const UProperty* InProperty) const
+{
+	bool bIsEditable = Super::CanEditChange(InProperty);
+	if (bIsEditable && InProperty)
+	{
+		const FName PropertyName = InProperty->GetFName();
+		if (PropertyName == GET_MEMBER_NAME_CHECKED(UBlueprint, bNativize))
+		{
+			return (BlueprintType != BPTYPE_LevelScript) && (BlueprintType != BPTYPE_MacroLibrary);
+		}
+	}
+	return bIsEditable;
+}
+#endif // WITH_EDITOR
+
 void UBlueprint::TagSubobjects(EObjectFlags NewFlags)
 {
 	Super::TagSubobjects(NewFlags);
@@ -1566,4 +1571,4 @@ UEdGraph* UBlueprint::GetLastEditedUberGraph() const
 	return nullptr;
 }
 
-#endif //WITH_EDITOR
\ No newline at end of file
+#endif //WITH_EDITOR
diff --git a/Engine/Source/Runtime/Engine/Private/BlueprintGeneratedClass.cpp b/Engine/Source/Runtime/Engine/Private/BlueprintGeneratedClass.cpp
--- a/Engine/Source/Runtime/Engine/Private/BlueprintGeneratedClass.cpp
+++ b/Engine/Source/Runtime/Engine/Private/BlueprintGeneratedClass.cpp
@@ -563,7 +563,10 @@ UObject* UBlueprintGeneratedClass::FindArchetype(UClass* ArchetypeClass, const F
 					ClassSCS->PreloadChain();
 				}
 
-				SCSNode = ClassSCS->FindSCSNode(ArchetypeBaseName);
+				// We keep the index name here rather than the base name, in order to avoid potential
+				// collisions between an SCS variable name and an existing AddComponent node template.
+				// This is because old AddComponent node templates were based on the class display name.
+				SCSNode = ClassSCS->FindSCSNode(ArchetypeName);
 			}
 
 			if (SCSNode)
@@ -580,9 +583,31 @@ UObject* UBlueprintGeneratedClass::FindArchetype(UClass* ArchetypeClass, const F
 			}
 			else if(UInheritableComponentHandler* ICH = Class->GetInheritableComponentHandler())
 			{
-				// Component template overrides should always match the archetype's base name (i.e. - not the instance name).
-				// This is because when we create new override objects within the ICH, we match the original template's name.
-				Archetype = ICH->GetOverridenComponentTemplate(ICH->FindKey(ArchetypeBaseName));
+				// This would find either an SCS component template override (for which the archetype
+				// name will match the SCS variable name), or an old AddComponent node template override
+				// (for which the archetype name will match the override record's component template name).
+				FComponentKey ComponentKey = ICH->FindKey(ArchetypeName);
+				if (!ComponentKey.IsValid() && ArchetypeName != ArchetypeBaseName)
+				{
+					// We didn't find either an SCS override or an old AddComponent template override,
+					// so now we look for a match with the base name; this would apply to new AddComponent
+					// node template overrides, which use the base name (non-index form).
+					ComponentKey = ICH->FindKey(ArchetypeBaseName);
+
+					// If we found a match with an SCS key instead, treat this as a collision and throw it
+					// out, because it should have already been found in the first search. This could happen
+					// if an old AddComponent node template's base name collides with an SCS variable name.
+					if (ComponentKey.IsValid() && ComponentKey.IsSCSKey())
+					{
+						ComponentKey = FComponentKey();
+					}
+				}
+
+				// Avoid searching for an invalid key.
+				if (ComponentKey.IsValid())
+				{
+					Archetype = ICH->GetOverridenComponentTemplate(ComponentKey);
+				}
 			}
 
 			if (Archetype == nullptr)
diff --git a/Engine/Source/Runtime/Engine/Private/BlueprintSetLibrary.cpp b/Engine/Source/Runtime/Engine/Private/BlueprintSetLibrary.cpp
--- a/Engine/Source/Runtime/Engine/Private/BlueprintSetLibrary.cpp
+++ b/Engine/Source/Runtime/Engine/Private/BlueprintSetLibrary.cpp
@@ -6,18 +6,14 @@
 #include "Kismet/KismetArrayLibrary.h"
 #include "Templates/HasGetTypeHash.h"
 
-UBlueprintSetLibrary::UBlueprintSetLibrary(const FObjectInitializer& ObjectInitializer)
-	: Super(ObjectInitializer)
-{
-}
-
-void UBlueprintSetLibrary::GenericSet_Add(const void* TargetSet, const USetProperty* SetProperty, const void* ItemPtr)
+bool UBlueprintSetLibrary::GenericSet_Add(const void* TargetSet, const USetProperty* SetProperty, const void* ItemPtr)
 {
 	if (TargetSet)
 	{
 		FScriptSetHelper SetHelper(SetProperty, TargetSet);
-		SetHelper.AddElement(ItemPtr);
+		return SetHelper.AddElement(ItemPtr);
 	}
+	return false;
 }
 
 void UBlueprintSetLibrary::GenericSet_AddItems(const void* TargetSet, const USetProperty* SetProperty, const void* TargetArray, const UArrayProperty* ArrayProperty)
@@ -34,13 +30,14 @@ void UBlueprintSetLibrary::GenericSet_AddItems(const void* TargetSet, const USet
 	}
 }
 
-void UBlueprintSetLibrary::GenericSet_Remove(const void* TargetSet, const USetProperty* SetProperty, const void* ItemPtr)
+bool UBlueprintSetLibrary::GenericSet_Remove(const void* TargetSet, const USetProperty* SetProperty, const void* ItemPtr)
 {
 	if (TargetSet)
 	{
 		FScriptSetHelper SetHelper(SetProperty, TargetSet);
-		SetHelper.RemoveElement(ItemPtr);
+		return SetHelper.RemoveElement(ItemPtr);
 	}
+	return false;
 }
 
 void UBlueprintSetLibrary::GenericSet_RemoveItems(const void* TargetSet, const USetProperty* SetProperty, const void* TargetArray, const UArrayProperty* ArrayProperty)
diff --git a/Engine/Source/Runtime/Engine/Private/KismetMathLibrary.cpp b/Engine/Source/Runtime/Engine/Private/KismetMathLibrary.cpp
--- a/Engine/Source/Runtime/Engine/Private/KismetMathLibrary.cpp
+++ b/Engine/Source/Runtime/Engine/Private/KismetMathLibrary.cpp
@@ -597,6 +597,13 @@ int32 UKismetMathLibrary::FTrunc(float A)
 	return FMath::TruncToInt(A);
 }	
 
+FIntVector UKismetMathLibrary::FTruncVector(const FVector& InVector)
+{
+	return FIntVector(FMath::TruncToInt(InVector.X),
+		              FMath::TruncToInt(InVector.Y),
+		              FMath::TruncToInt(InVector.Z));
+}
+
 int32 UKismetMathLibrary::FCeil(float A)
 {
 	return FMath::CeilToInt(A);
@@ -1897,6 +1904,11 @@ uint8 UKismetMathLibrary::Conv_IntToByte(int32 InInt)
 	return (uint8)InInt;
 }
 
+FIntVector UKismetMathLibrary::Conv_IntToIntVector(int32 InInt)
+{
+	return FIntVector(InInt, InInt, InInt);
+}
+
 bool UKismetMathLibrary::Conv_IntToBool(int32 InInt)
 {
 	return InInt == 0 ? false : true;
@@ -1942,6 +1954,11 @@ FVector2D UKismetMathLibrary::Conv_VectorToVector2D(FVector InVec)
 	return FVector2D(InVec);
 }
 
+FVector UKismetMathLibrary::Conv_IntVectorToVector(const FIntVector& InIntVector)
+{
+	return FVector(InIntVector);
+}
+
 FVector UKismetMathLibrary::Conv_Vector2DToVector(FVector2D InVec2D, float Z)
 {
 	return FVector(InVec2D, Z);
diff --git a/Engine/Source/Runtime/Engine/Private/KismetStringLibrary.cpp b/Engine/Source/Runtime/Engine/Private/KismetStringLibrary.cpp
--- a/Engine/Source/Runtime/Engine/Private/KismetStringLibrary.cpp
+++ b/Engine/Source/Runtime/Engine/Private/KismetStringLibrary.cpp
@@ -72,6 +72,11 @@ FString UKismetStringLibrary::Conv_VectorToString(FVector InVec)
 	return InVec.ToString();	
 }
 
+FString UKismetStringLibrary::Conv_IntVectorToString(FIntVector InIntVec)
+{
+	return InIntVec.ToString();
+}
+
 FString UKismetStringLibrary::Conv_Vector2dToString(FVector2D InVec)
 {
 	return InVec.ToString();	
@@ -197,6 +202,21 @@ FString UKismetStringLibrary::BuildString_Vector(const FString& AppendTo, const
 	return StringResult;
 }
 
+FString UKismetStringLibrary::BuildString_IntVector(const FString& AppendTo, const FString& Prefix, FIntVector InIntVector, const FString& Suffix)
+{
+	// faster, preallocating method
+	FString const VecStr = InIntVector.ToString();
+
+	FString StringResult;
+	StringResult.Empty(AppendTo.Len() + Prefix.Len() + VecStr.Len() + Suffix.Len() + 1); // adding one for the string terminator
+	StringResult += AppendTo;
+	StringResult += Prefix;
+	StringResult += VecStr;
+	StringResult += Suffix;
+
+	return StringResult;
+}
+
 FString UKismetStringLibrary::BuildString_Vector2d(const FString& AppendTo, const FString& Prefix, FVector2D InVector2d, const FString& Suffix)
 {
 	// faster, preallocating method
